Move the Subtitles plugin from Restricted to Experimental.

#virtualized

diff --git a/Engine/Config/BaseEngine.ini b/Engine/Config/BaseEngine.ini
--- a/Engine/Config/BaseEngine.ini
+++ b/Engine/Config/BaseEngine.ini
@@ -3885,4 +3885,7 @@ RuntimeHashDefaultClass=/Script/CoreUObject.Class'/Script/Engine.WorldPartitionR
 VerificationThreadCount=4
 
 [/Script/ControlRig.ControlRigSettings]
-DefaultRootModule=/ControlRig/Modules/Modules56/Root.Root
\ No newline at end of file
+DefaultRootModule=/ControlRig/Modules/Modules56/Root.Root
+
+[/Script/SubtitlesAndClosedCaptions.SubtitlesSettings]
+SubtitleWidgetToUseDefault=/Script/UMG.WidgetBlueprintGeneratedClass'/SubtitlesAndClosedCaptions/Widgets/DefaultSubtitleWidget.DefaultSubtitleWidget_C'
\ No newline at end of file
diff --git a/Engine/Plugins/Experimental/SubtitlesAndClosedCaptions/Source/SubtitlesAndClosedCaptions/Private/MovieSceneSubtitleComponentTypes.cpp b/Engine/Plugins/Experimental/SubtitlesAndClosedCaptions/Source/SubtitlesAndClosedCaptions/Private/MovieSceneSubtitleComponentTypes.cpp
new file mode 100644
--- /dev/null
+++ b/Engine/Plugins/Experimental/SubtitlesAndClosedCaptions/Source/SubtitlesAndClosedCaptions/Private/MovieSceneSubtitleComponentTypes.cpp
@@ -0,0 +1,36 @@
+// Copyright Epic Games, Inc. All Rights Reserved.
+
+#include "MovieSceneSubtitleComponentTypes.h"
+
+#include "EntitySystem/MovieSceneComponentRegistry.h"
+#include "EntitySystem/MovieSceneEntityFactoryTemplates.h"
+#include "EntitySystem/MovieSceneEntitySystemLinker.h"
+
+bool FMovieSceneSubtitleComponentTypes::GMovieSceneSubtitlesComponentTypesDestroyed = false;
+TUniquePtr<FMovieSceneSubtitleComponentTypes> FMovieSceneSubtitleComponentTypes::GMovieSceneSubtitlesComponentTypes;
+
+
+FMovieSceneSubtitleComponentTypes* FMovieSceneSubtitleComponentTypes::Get()
+{
+	if (!GMovieSceneSubtitlesComponentTypes.IsValid())
+	{
+		check(!GMovieSceneSubtitlesComponentTypesDestroyed);
+		GMovieSceneSubtitlesComponentTypes.Reset(new FMovieSceneSubtitleComponentTypes);
+	}
+	return GMovieSceneSubtitlesComponentTypes.Get();
+}
+
+void FMovieSceneSubtitleComponentTypes::Destroy()
+{
+	GMovieSceneSubtitlesComponentTypes.Reset();
+	GMovieSceneSubtitlesComponentTypesDestroyed = true;
+}
+
+FMovieSceneSubtitleComponentTypes::FMovieSceneSubtitleComponentTypes()
+{
+	using namespace UE::MovieScene;
+	FComponentRegistry* ComponentRegistry = UMovieSceneEntitySystemLinker::GetComponents();
+	check(ComponentRegistry);
+	
+	ComponentRegistry->NewComponentType(&SubtitleData, TEXT("SubtitleData Component"));
+}
diff --git a/Engine/Plugins/Experimental/SubtitlesAndClosedCaptions/Source/SubtitlesAndClosedCaptions/Private/MovieSceneSubtitleSection.cpp b/Engine/Plugins/Experimental/SubtitlesAndClosedCaptions/Source/SubtitlesAndClosedCaptions/Private/MovieSceneSubtitleSection.cpp
new file mode 100644
--- /dev/null
+++ b/Engine/Plugins/Experimental/SubtitlesAndClosedCaptions/Source/SubtitlesAndClosedCaptions/Private/MovieSceneSubtitleSection.cpp
@@ -0,0 +1,66 @@
+// Copyright Epic Games, Inc. All Rights Reserved.
+
+#include "MovieSceneSubtitleSection.h"
+
+#include "EntitySystem/BuiltInComponentTypes.h"
+#include "EntitySystem/MovieSceneEntitySystemLinker.h"
+#include "MovieScene.h"
+#include "MovieSceneFwd.h"
+#include "MovieSceneSubtitleComponentTypes.h"
+#include "SubtitleDataComponent.h"
+#include "Subtitles/SubtitlesAndClosedCaptionsDelegates.h"
+
+#include UE_INLINE_GENERATED_CPP_BY_NAME(MovieSceneSubtitleSection)
+
+using namespace UE::MovieScene;
+
+UMovieSceneSubtitleSection::UMovieSceneSubtitleSection(const FObjectInitializer& ObjectInitializer)
+	: Super(ObjectInitializer)
+{
+}
+
+TOptional<TRange<FFrameNumber>> UMovieSceneSubtitleSection::GetAutoSizeRange() const
+{
+	if (!Subtitle)
+	{
+		return TRange<FFrameNumber>();
+	}
+
+	const FFrameRate FrameRate = GetTypedOuter<UMovieScene>()->GetTickResolution();
+	const FFrameTime DurationToUse = Subtitle->Duration * FrameRate;
+	return TRange<FFrameNumber>(GetInclusiveStartFrame(), GetInclusiveStartFrame() + DurationToUse.FrameNumber);
+}
+
+bool UMovieSceneSubtitleSection::PopulateEvaluationFieldImpl(const TRange<FFrameNumber>& EffectiveRange, const FMovieSceneEvaluationFieldEntityMetaData& InMetaData, FMovieSceneEntityComponentFieldBuilder* OutFieldBuilder)
+{
+	// check preconditions
+	check(OutFieldBuilder);
+
+	OutFieldBuilder->AddPersistentEntity(EffectiveRange, this);
+	return true;
+}
+
+void UMovieSceneSubtitleSection::ImportEntityImpl(UMovieSceneEntitySystemLinker* EntityLinker, const FEntityImportParams& Params, FImportedEntity* OutImportedEntity)
+{
+	// check preconditions
+	check(OutImportedEntity);
+
+	if (!Subtitle)
+	{
+		return;
+	}
+
+	const FBuiltInComponentTypes* BuiltInComponents = FBuiltInComponentTypes::Get();
+	check(BuiltInComponents);
+
+	FMovieSceneSubtitleComponentTypes* SubtitleComponents = FMovieSceneSubtitleComponentTypes::Get();
+	check(SubtitleComponents);
+
+	const FGuid ObjectBindingID = Params.GetObjectBindingID();
+
+	OutImportedEntity->AddBuilder(
+		FEntityBuilder()
+		.AddConditional(BuiltInComponents->GenericObjectBinding, ObjectBindingID, ObjectBindingID.IsValid())
+		.AddTagConditional(BuiltInComponents->Tags.Root, !ObjectBindingID.IsValid())
+		.Add(SubtitleComponents->SubtitleData, FSubtitleDataComponent{ this, EMovieScenePlayerStatus::Type::Stopped }));
+}
diff --git a/Engine/Plugins/Experimental/SubtitlesAndClosedCaptions/Source/SubtitlesAndClosedCaptions/Private/MovieSceneSubtitlesSystem.cpp b/Engine/Plugins/Experimental/SubtitlesAndClosedCaptions/Source/SubtitlesAndClosedCaptions/Private/MovieSceneSubtitlesSystem.cpp
new file mode 100644
--- /dev/null
+++ b/Engine/Plugins/Experimental/SubtitlesAndClosedCaptions/Source/SubtitlesAndClosedCaptions/Private/MovieSceneSubtitlesSystem.cpp
@@ -0,0 +1,68 @@
+// Copyright Epic Games, Inc. All Rights Reserved.
+
+#include "MovieSceneSubtitlesSystem.h"
+
+#include "EntitySystem/BuiltInComponentTypes.h"
+
+#include "MovieSceneSubtitleComponentTypes.h"
+#include "MovieSceneSubtitleSection.h"
+
+#include "SubtitlesAndClosedCaptionsModule.h"
+
+using namespace UE::MovieScene;
+
+
+UMovieSceneSubtitlesSystem::UMovieSceneSubtitlesSystem(const FObjectInitializer& ObjInit)
+	: UMovieSceneEntitySystem(ObjInit)
+{
+	FMovieSceneSubtitleComponentTypes* SubtitlesComponents = FMovieSceneSubtitleComponentTypes::Get();
+	check(SubtitlesComponents);
+
+	RelevantComponent = SubtitlesComponents->SubtitleData;
+
+	Phase = ESystemPhase::Scheduling;
+}
+
+UMovieSceneSubtitlesSystem::~UMovieSceneSubtitlesSystem()
+{
+}
+
+void UMovieSceneSubtitlesSystem::OnSchedulePersistentTasks(UE::MovieScene::IEntitySystemScheduler* TaskScheduler)
+{
+	const FBuiltInComponentTypes* BuiltInComponents = FBuiltInComponentTypes::Get();
+	check(BuiltInComponents);
+
+	const FMovieSceneSubtitleComponentTypes* SubtitleComponents = FMovieSceneSubtitleComponentTypes::Get();
+	check(SubtitleComponents);
+
+	check(Linker);
+	const FInstanceRegistry* InstanceRegistry = Linker->GetInstanceRegistry();
+	check(InstanceRegistry);
+
+	FTaskID EvaluateTask = FEntityTaskBuilder()
+		.Read(BuiltInComponents->InstanceHandle)
+		.Write(SubtitleComponents->SubtitleData)
+		.SetDesiredThread(ENamedThreads::GameThread)
+		.Schedule_PerAllocation<FEvaluateSubtitles>(&Linker->EntityManager, TaskScheduler, *InstanceRegistry);
+}
+
+void UMovieSceneSubtitlesSystem::OnRun(FSystemTaskPrerequisites& InPrerequisites, FSystemSubsequentTasks& Subsequents)
+{
+	const FBuiltInComponentTypes* BuiltInComponents = FBuiltInComponentTypes::Get();
+	check(BuiltInComponents);
+
+	const FMovieSceneSubtitleComponentTypes* SubtitleComponents = FMovieSceneSubtitleComponentTypes::Get();
+	check(SubtitleComponents);
+
+	check(Linker);
+	const FInstanceRegistry* InstanceRegistry = Linker->GetInstanceRegistry();
+	check(InstanceRegistry);
+
+	FEntityTaskBuilder()
+		.Read(BuiltInComponents->InstanceHandle)
+		.Write(SubtitleComponents->SubtitleData)
+		.SetDesiredThread(ENamedThreads::GameThread)
+		.template Dispatch_PerAllocation<FEvaluateSubtitles>(&Linker->EntityManager, InPrerequisites, &Subsequents, *InstanceRegistry);
+
+
+}
diff --git a/Engine/Plugins/Experimental/SubtitlesAndClosedCaptions/Source/SubtitlesAndClosedCaptions/Private/MovieSceneSubtitlesTrack.cpp b/Engine/Plugins/Experimental/SubtitlesAndClosedCaptions/Source/SubtitlesAndClosedCaptions/Private/MovieSceneSubtitlesTrack.cpp
new file mode 100644
--- /dev/null
+++ b/Engine/Plugins/Experimental/SubtitlesAndClosedCaptions/Source/SubtitlesAndClosedCaptions/Private/MovieSceneSubtitlesTrack.cpp
@@ -0,0 +1,91 @@
+// Copyright Epic Games, Inc. All Rights Reserved.
+
+#include "MovieSceneSubtitlesTrack.h"
+
+#include "MovieScene.h"
+#include "MovieSceneCommonHelpers.h"
+#include "MovieSceneSubtitleSection.h"
+#include "Subtitles/SubtitlesAndClosedCaptionsDelegates.h"
+
+#include UE_INLINE_GENERATED_CPP_BY_NAME(MovieSceneSubtitlesTrack)
+
+#define LOCTEXT_NAMESPACE "MovieSceneSubtitlesTrack"
+
+
+UMovieSceneSubtitlesTrack::UMovieSceneSubtitlesTrack(const FObjectInitializer& ObjectInitializer)
+	: Super(ObjectInitializer)
+{
+}
+
+UMovieSceneSection* UMovieSceneSubtitlesTrack::AddNewSubtitle(const USubtitleAssetUserData& Subtitle, FFrameNumber InStartTime)
+{
+	return AddNewSubtitleOnRow(Subtitle, InStartTime, INDEX_NONE);
+}
+
+UMovieSceneSection* UMovieSceneSubtitlesTrack::AddNewSubtitleOnRow(const USubtitleAssetUserData& Subtitle, FFrameNumber InStartTime, int32 InRowIndex)
+{
+	UMovieSceneSubtitleSection* NewSection = Cast<UMovieSceneSubtitleSection>(CreateNewSection());
+
+	const FFrameRate FrameRate = GetTypedOuter<UMovieScene>()->GetTickResolution();
+	const FFrameTime DurationToUse = Subtitle.Duration * FrameRate;
+
+	ensure(DurationToUse.FrameNumber.Value > 0);
+	NewSection->InitialPlacementOnRow(SubtitleSections, InStartTime, DurationToUse.FrameNumber.Value, InRowIndex);
+	NewSection->SetSubtitle(Subtitle);
+
+	SubtitleSections.Add(NewSection);
+
+	return NewSection;
+}
+
+bool UMovieSceneSubtitlesTrack::SupportsType(TSubclassOf<UMovieSceneSection> SectionClass) const
+{
+	return SectionClass == UMovieSceneSubtitleSection::StaticClass();
+}
+
+void UMovieSceneSubtitlesTrack::RemoveAllAnimationData()
+{
+	return SubtitleSections.Empty();
+}
+
+bool UMovieSceneSubtitlesTrack::HasSection(const UMovieSceneSection& Section) const
+{
+	return SubtitleSections.Contains(&Section);
+}
+
+void UMovieSceneSubtitlesTrack::AddSection(UMovieSceneSection& Section)
+{
+	SubtitleSections.Add(&Section);
+}
+
+void UMovieSceneSubtitlesTrack::RemoveSection(UMovieSceneSection& Section)
+{
+	SubtitleSections.Remove(&Section);
+}
+
+void UMovieSceneSubtitlesTrack::RemoveSectionAt(int32 SectionIndex)
+{
+	SubtitleSections.RemoveAt(SectionIndex);
+}
+
+bool UMovieSceneSubtitlesTrack::IsEmpty() const
+{
+	return SubtitleSections.IsEmpty();
+}
+
+const TArray<UMovieSceneSection*>& UMovieSceneSubtitlesTrack::GetAllSections() const
+{
+	return SubtitleSections;
+}
+
+bool UMovieSceneSubtitlesTrack::SupportsMultipleRows() const
+{
+	return true;
+}
+
+UMovieSceneSection* UMovieSceneSubtitlesTrack::CreateNewSection()
+{
+	return NewObject<UMovieSceneSubtitleSection>(this, NAME_None, RF_Transactional);
+}
+
+#undef LOCTEXT_NAMESPACE
diff --git a/Engine/Plugins/Experimental/SubtitlesAndClosedCaptions/Source/SubtitlesAndClosedCaptions/Private/SubtitleTextBlock.cpp b/Engine/Plugins/Experimental/SubtitlesAndClosedCaptions/Source/SubtitlesAndClosedCaptions/Private/SubtitleTextBlock.cpp
new file mode 100644
--- /dev/null
+++ b/Engine/Plugins/Experimental/SubtitlesAndClosedCaptions/Source/SubtitlesAndClosedCaptions/Private/SubtitleTextBlock.cpp
@@ -0,0 +1,100 @@
+// Copyright Epic Games, Inc. All Rights Reserved.
+
+#include "SubtitleTextBlock.h"
+#include "SubtitlesAndClosedCaptionsModule.h"
+
+void USubtitleTextBlock::NativeConstruct()
+{
+	Super::NativeConstruct();
+
+	// Ensure the text blocks don't start visible.
+	if (IsValid(DialogSubtitleBlock))
+	{
+		DialogSubtitleBlock->SetVisibility(ESlateVisibility::Collapsed);
+	}
+
+	if (IsValid(CaptionSubtitleBlock))
+	{
+		CaptionSubtitleBlock->SetVisibility(ESlateVisibility::Collapsed);
+	}
+
+	if (IsValid(DescriptionSubtitleBlock))
+	{
+		DescriptionSubtitleBlock->SetVisibility(ESlateVisibility::Collapsed);
+	}
+}
+
+void USubtitleTextBlock::StartDisplayingSubtitle(const USubtitleAssetUserData& Subtitle)
+{
+	// Pick which UextBlock is relevant by category....
+	UTextBlock* TextToModify = nullptr;
+
+	switch (Subtitle.SubtitleType)
+	{
+	default:
+	{
+		UE_LOG(LogSubtitlesAndClosedCaptions, Warning, TEXT("An unrecognized subtitle type was queued for display. Using the standard Subtitle text block as a fallback."));
+		// Fallthrough, no break.
+	}
+
+	case ESubtitleType::Subtitle:
+	{
+		TextToModify = DialogSubtitleBlock;
+		break;
+	}
+
+	case ESubtitleType::ClosedCaption:
+	{
+		TextToModify = CaptionSubtitleBlock;
+		break;
+	}
+
+	case ESubtitleType::AudioDescription:
+	{
+		TextToModify = DescriptionSubtitleBlock;
+		break;
+	}
+
+	}
+
+	// ...Then modify and display it
+	if (IsValid(TextToModify))
+	{
+		TextToModify->SetText(Subtitle.Text);
+		TextToModify->SetVisibility(ESlateVisibility::HitTestInvisible);
+	}
+}
+
+void USubtitleTextBlock::StopDisplayingSubtitle(const ESubtitleType SubtitleType)
+{
+
+	UTextBlock* TextToModify = nullptr;
+	switch (SubtitleType)
+	{
+	default: // Fallthrough, no break (see ::StartDisplayingSubtitle)
+	case ESubtitleType::Subtitle:
+	{
+		TextToModify = DialogSubtitleBlock;
+		break;
+	}
+
+	case ESubtitleType::ClosedCaption:
+	{
+		TextToModify = CaptionSubtitleBlock;
+		break;
+	}
+
+	case ESubtitleType::AudioDescription:
+	{
+		TextToModify = DescriptionSubtitleBlock;
+		break;
+	}
+
+	}
+
+
+	if (IsValid(TextToModify))
+	{
+		TextToModify->SetVisibility(ESlateVisibility::Collapsed);
+	}
+}
diff --git a/Engine/Plugins/Experimental/SubtitlesAndClosedCaptions/Source/SubtitlesAndClosedCaptions/Private/SubtitlesAndClosedCaptionsModule.cpp b/Engine/Plugins/Experimental/SubtitlesAndClosedCaptions/Source/SubtitlesAndClosedCaptions/Private/SubtitlesAndClosedCaptionsModule.cpp
new file mode 100644
--- /dev/null
+++ b/Engine/Plugins/Experimental/SubtitlesAndClosedCaptions/Source/SubtitlesAndClosedCaptions/Private/SubtitlesAndClosedCaptionsModule.cpp
@@ -0,0 +1,26 @@
+// Copyright Epic Games, Inc. All Rights Reserved.
+
+#include "SubtitlesAndClosedCaptionsModule.h"
+
+#include "Features/IModularFeatures.h"
+#include "Modules/ModuleManager.h"
+#include "Subtitles/SubtitlesAndClosedCaptionsDelegates.h"
+
+DEFINE_LOG_CATEGORY(LogSubtitlesAndClosedCaptions);
+
+void FSubtitlesAndClosedCaptionsModule::StartupModule()
+{
+	IModularFeatures::Get().RegisterModularFeature(ISubtitlesAndClosedCaptionsModule::GetModularFeatureName(), this);
+}
+
+void FSubtitlesAndClosedCaptionsModule::ShutdownModule()
+{
+	IModularFeatures::Get().UnregisterModularFeature(ISubtitlesAndClosedCaptionsModule::GetModularFeatureName(), this);
+}
+
+TSubclassOf<UAssetUserData> FSubtitlesAndClosedCaptionsModule::GetAssetUserDataClass() const
+{
+	return USubtitleAssetUserData::StaticClass();
+}
+
+IMPLEMENT_MODULE(FSubtitlesAndClosedCaptionsModule, SubtitlesAndClosedCaptions);
diff --git a/Engine/Plugins/Experimental/SubtitlesAndClosedCaptions/Source/SubtitlesAndClosedCaptions/Private/SubtitlesBlueprintFunctionLibrary.cpp b/Engine/Plugins/Experimental/SubtitlesAndClosedCaptions/Source/SubtitlesAndClosedCaptions/Private/SubtitlesBlueprintFunctionLibrary.cpp
new file mode 100644
--- /dev/null
+++ b/Engine/Plugins/Experimental/SubtitlesAndClosedCaptions/Source/SubtitlesAndClosedCaptions/Private/SubtitlesBlueprintFunctionLibrary.cpp
@@ -0,0 +1,43 @@
+// Copyright Epic Games, Inc. All Rights Reserved.
+
+#include "SubtitlesBlueprintFunctionLibrary.h"
+
+#include "Subtitles/SubtitlesAndClosedCaptionsDelegates.h"
+
+void USubtitlesBlueprintFunctionLibrary::QueueSubtitle(const USubtitleAssetUserData* Subtitle, const ESubtitleTiming Timing)
+{
+	if (!Subtitle)
+	{
+		return;
+	}
+
+	FQueueSubtitleParameters Params{ *CastChecked<const UAssetUserData>(Subtitle) };
+	FSubtitlesAndClosedCaptionsDelegates::QueueSubtitle.ExecuteIfBound(Params, Timing);
+}
+
+bool USubtitlesBlueprintFunctionLibrary::IsSubtitleActive(const USubtitleAssetUserData* Subtitle)
+{
+	if (!Subtitle)
+	{
+		return false;
+	}
+	if (!FSubtitlesAndClosedCaptionsDelegates::IsSubtitleActive.IsBound())
+	{
+		return false;
+	}
+	return FSubtitlesAndClosedCaptionsDelegates::IsSubtitleActive.Execute(*CastChecked<const UAssetUserData>(Subtitle));
+}
+
+void USubtitlesBlueprintFunctionLibrary::StopSubtitle(const USubtitleAssetUserData* Subtitle)
+{
+	if (!Subtitle)
+	{
+		return;
+	}
+	FSubtitlesAndClosedCaptionsDelegates::StopSubtitle.ExecuteIfBound(*CastChecked<const UAssetUserData>(Subtitle));
+}
+
+void USubtitlesBlueprintFunctionLibrary::StopAllSubtitles()
+{
+	FSubtitlesAndClosedCaptionsDelegates::StopAllSubtitles.ExecuteIfBound();
+}
diff --git a/Engine/Plugins/Experimental/SubtitlesAndClosedCaptions/Source/SubtitlesAndClosedCaptions/Private/SubtitlesSettings.cpp b/Engine/Plugins/Experimental/SubtitlesAndClosedCaptions/Source/SubtitlesAndClosedCaptions/Private/SubtitlesSettings.cpp
new file mode 100644
--- /dev/null
+++ b/Engine/Plugins/Experimental/SubtitlesAndClosedCaptions/Source/SubtitlesAndClosedCaptions/Private/SubtitlesSettings.cpp
@@ -0,0 +1,10 @@
+// Copyright Epic Games, Inc. All Rights Reserved.
+#include "SubtitlesSettings.h"
+
+#define LOCTEXT_NAMESPACE "SubtitlesSettings"
+
+USubtitlesSettings::USubtitlesSettings()
+{
+}
+
+#undef LOCTEXT_NAMESPACE
diff --git a/Engine/Plugins/Experimental/SubtitlesAndClosedCaptions/Source/SubtitlesAndClosedCaptions/Private/SubtitlesSubsystem.cpp b/Engine/Plugins/Experimental/SubtitlesAndClosedCaptions/Source/SubtitlesAndClosedCaptions/Private/SubtitlesSubsystem.cpp
new file mode 100644
--- /dev/null
+++ b/Engine/Plugins/Experimental/SubtitlesAndClosedCaptions/Source/SubtitlesAndClosedCaptions/Private/SubtitlesSubsystem.cpp
@@ -0,0 +1,277 @@
+// Copyright Epic Games, Inc. All Rights Reserved.
+
+#include "SubtitlesSubsystem.h"
+
+#include "Algo/RemoveIf.h"
+#include "CanvasItem.h"
+#include "CanvasTypes.h"
+#include "Containers/Ticker.h"
+#include "Engine/Font.h"
+#include "Engine/World.h"
+#include "Math/UnrealMathUtility.h"
+#include "Styling/CoreStyle.h"
+#include "SubtitlesAndClosedCaptionsModule.h"
+#include "SubtitlesSettings.h"
+#include "TimerManager.h"
+#include "UnrealClient.h"
+#include "UObject/ObjectMacros.h"
+#include "UObject/Package.h"
+#include "UObject/StrongObjectPtrTemplates.h"
+#include "UObject/WeakObjectPtrTemplates.h"
+
+// #SUBTITLES_PRD / #SUBTITLES_PRD_TODO - remove subtitles development comments prior to release
+
+void USubtitlesSubsystem::Initialize(FSubsystemCollectionBase& Collection)
+{
+	if (TryCreateUMGWidget())
+	{
+		BindDelegates();
+	}
+}
+
+void USubtitlesSubsystem::BindDelegates()
+{
+	FSubtitlesAndClosedCaptionsDelegates::QueueSubtitle.BindUObject(this, &USubtitlesSubsystem::QueueSubtitle);
+	FSubtitlesAndClosedCaptionsDelegates::IsSubtitleActive.BindUObject(this, &USubtitlesSubsystem::IsSubtitleActive);
+	FSubtitlesAndClosedCaptionsDelegates::StopSubtitle.BindUObject(this, &USubtitlesSubsystem::StopSubtitle);
+	FSubtitlesAndClosedCaptionsDelegates::StopAllSubtitles.BindUObject(this, &USubtitlesSubsystem::StopAllSubtitles);
+}
+
+
+void USubtitlesSubsystem::QueueSubtitle(const FQueueSubtitleParameters& Params, const ESubtitleTiming Timing)
+{
+	constexpr float InfiniteDuration = FLT_MAX;
+	const USubtitleAssetUserData& Subtitle = *CastChecked<const USubtitleAssetUserData>(&Params.Subtitle);
+	UE_LOG(LogSubtitlesAndClosedCaptions, Display, TEXT("QueueSubtitle: '%s'"), *Subtitle.Text.ToString());
+
+
+	// Externally-timed subtitles will be removed by the system queueing them, so they have an otherwise-infinite duration.
+	float Duration = InfiniteDuration;
+	if (Timing == ESubtitleTiming::InternallyTimed)
+	{
+		Duration = Params.Duration.IsSet() ? Params.Duration.GetValue() : Subtitle.Duration;
+	}
+
+	if (IsInGameThread())
+	{
+		AddActiveSubtitle(Subtitle, Duration);
+	}
+	else
+	{
+		ExecuteOnGameThread(
+			TEXT("USubtitlesSubsystem::HandleQueueSubtitle"),
+			[WeakThis = TWeakObjectPtr<USubtitlesSubsystem>(this), WeakSubtitle = TWeakObjectPtr<const USubtitleAssetUserData>(&Subtitle), Duration]()
+			{
+				if (TStrongObjectPtr<USubtitlesSubsystem> ThisPin = WeakThis.Pin())
+				{
+					if (TStrongObjectPtr<const USubtitleAssetUserData> SubtitlePin = WeakSubtitle.Pin())
+					{
+						ThisPin->AddActiveSubtitle(*SubtitlePin, Duration);
+					}
+				}
+			}
+		);
+	}
+}
+
+void USubtitlesSubsystem::AddActiveSubtitle(const USubtitleAssetUserData& Subtitle, float Duration)
+{
+	// If the subtitle is already active then update its duration (by removing it then and then re-adding it)
+	const FActiveSubtitle* FoundActiveSubtitle = ActiveSubtitles.FindByPredicate(
+		[SubtitlePtr = &Subtitle](const FActiveSubtitle& ActiveSubtitle) { return ActiveSubtitle.Subtitle == SubtitlePtr; });
+	if (FoundActiveSubtitle != nullptr)
+	{
+		RemoveActiveSubtitle(&Subtitle);
+	}
+
+	FActiveSubtitle NewActiveSubtitle{ &Subtitle };
+
+	// Subtitles with delayed offset need a timer to await their entry in the queue.
+	if (Subtitle.StartOffset > 0.f)
+	{
+		// The timer handle will be reused for the duration when it enters the active subtitle queue.
+		// For now the timer tracks how long until it enters that queue.
+		FTimerDelegate Delegate = FTimerDelegate::CreateUObject(this, &USubtitlesSubsystem::MakeDelayedSubtitleActive, NewActiveSubtitle.Subtitle);
+		GetWorldRef().GetTimerManager().SetTimer(NewActiveSubtitle.DurationTimerHandle, MoveTemp(Delegate), Subtitle.StartOffset, /*bLoop=*/false);
+
+		DelayedSubtitles.Add(MoveTemp(NewActiveSubtitle));
+	}
+	else
+	{
+		// Without the delayed offset, instantly enter the queue as usual.
+		// The timer here tracks how long until the subtitle will expire and leave the active subtitle queue.
+		FTimerDelegate Delegate = FTimerDelegate::CreateUObject(this, &USubtitlesSubsystem::RemoveActiveSubtitle, NewActiveSubtitle.Subtitle);
+		Duration = FMath::Max(Duration, SubtitleMinDuration);
+		GetWorldRef().GetTimerManager().SetTimer(NewActiveSubtitle.DurationTimerHandle, MoveTemp(Delegate), Duration, /*bLoop=*/false);
+
+		ActiveSubtitles.Add(MoveTemp(NewActiveSubtitle));
+		ActiveSubtitles.StableSort([](const FActiveSubtitle& Lhs, const FActiveSubtitle& Rhs) { return Lhs.Subtitle->Priority > Rhs.Subtitle->Priority; });
+
+		UpdateWidgetData();
+	}
+}
+
+void USubtitlesSubsystem::MakeDelayedSubtitleActive(TObjectPtr<const USubtitleAssetUserData> Subtitle)
+{
+	if (IsValid(Subtitle))
+	{
+		FActiveSubtitle* DelayedSubtitle = Algo::FindByPredicate(DelayedSubtitles, [&Subtitle](const FActiveSubtitle& DelayedSubtitle) {return DelayedSubtitle.Subtitle == Subtitle; });
+
+		if (DelayedSubtitle != nullptr)
+		{
+			const USubtitleAssetUserData& SubtitleAsset = *DelayedSubtitle->Subtitle;
+			const float Duration = FMath::Max(SubtitleAsset.Duration, SubtitleMinDuration);
+
+			// Reuse the Timer Handle for duration, now that it's no longer needed for the delay.
+			FTimerDelegate Delegate = FTimerDelegate::CreateUObject(this, &USubtitlesSubsystem::RemoveActiveSubtitle, DelayedSubtitle->Subtitle);
+			GetWorldRef().GetTimerManager().SetTimer(DelayedSubtitle->DurationTimerHandle, MoveTemp(Delegate), Duration, /*bLoop=*/false);
+
+			// Insert the new subtitle to the actual queue and ensure it remains sorted by priority.
+			ActiveSubtitles.Add(MoveTemp(*DelayedSubtitle));
+			ActiveSubtitles.StableSort([](const FActiveSubtitle& Lhs, const FActiveSubtitle& Rhs) { return Lhs.Subtitle->Priority > Rhs.Subtitle->Priority; });
+
+			// Remove from the list of Delayed Subtitles
+			const int32 FirstRemovedIndex = Algo::RemoveIf(DelayedSubtitles, [&Subtitle](const FActiveSubtitle& DelayedSubtitle) {return DelayedSubtitle.Subtitle == Subtitle; });
+			DelayedSubtitles.SetNum(FirstRemovedIndex);
+
+			UpdateWidgetData();
+		}
+	}
+}
+
+bool USubtitlesSubsystem::IsSubtitleActive(const UAssetUserData& Data) const
+{
+	if (!ensureMsgf(IsInGameThread(), TEXT("IsSubtitleActive must currently be run on the GameThread - ActiveSubtitles vector is not locked")))
+	{
+		return false;
+	}
+
+	const USubtitleAssetUserData* Subtitle = CastChecked<const USubtitleAssetUserData>(&Data);
+
+	const FActiveSubtitle* FoundActiveSubtitle = ActiveSubtitles.FindByPredicate(
+		[Subtitle](const FActiveSubtitle& ActiveSubtitle) { return ActiveSubtitle.Subtitle == Subtitle; });
+
+	return FoundActiveSubtitle != nullptr;
+}
+
+void USubtitlesSubsystem::StopSubtitle(const UAssetUserData& Data)
+{
+	const USubtitleAssetUserData* Subtitle = CastChecked<const USubtitleAssetUserData>(&Data);
+	RemoveActiveSubtitle(Subtitle);
+}
+
+void USubtitlesSubsystem::StopAllSubtitles()
+{
+	// Clean up queued subtitles.
+	FTimerManager& TimerManager = GetWorldRef().GetTimerManager();
+	for (FActiveSubtitle& ActiveSubtitle : ActiveSubtitles)
+	{
+		TimerManager.ClearTimer(ActiveSubtitle.DurationTimerHandle);
+	}
+	ActiveSubtitles.Empty();
+
+	// Also remove delayed-start subtitles not yet in the queue.
+	for (FActiveSubtitle& DelayedSubtitle : DelayedSubtitles)
+	{
+		TimerManager.ClearTimer(DelayedSubtitle.DurationTimerHandle);
+	}
+	DelayedSubtitles.Empty();
+
+	// Clear the widget's display
+	if (IsValid(SubtitleWidget))
+	{
+		SubtitleWidget->StopDisplayingSubtitle(ESubtitleType::AudioDescription);
+		SubtitleWidget->StopDisplayingSubtitle(ESubtitleType::ClosedCaption);
+		SubtitleWidget->StopDisplayingSubtitle(ESubtitleType::Subtitle);
+	}
+	else 
+	{
+		UE_LOG(LogSubtitlesAndClosedCaptions, Warning, TEXT("Can't remove subtitles because there isn't a valid UMG widget."));
+	}
+}
+
+void USubtitlesSubsystem::RemoveActiveSubtitle(TObjectPtr<const USubtitleAssetUserData> Subtitle)
+{
+	bool bSuccessfullyRemoved = false;
+	int32 FirstRemovedIndex = Algo::StableRemoveIf(ActiveSubtitles, [&Subtitle](const FActiveSubtitle& ActiveSubtitle) {return ActiveSubtitle.Subtitle == Subtitle; });
+
+	FTimerManager& TimerManager = GetWorldRef().GetTimerManager();
+	for (int32 Index = FirstRemovedIndex; Index < ActiveSubtitles.Num(); ++Index)
+	{
+		TimerManager.ClearTimer(ActiveSubtitles[Index].DurationTimerHandle);
+		bSuccessfullyRemoved = true;
+	}
+
+	ActiveSubtitles.SetNum(FirstRemovedIndex);
+
+	// Stop Displaying the removed subtitle and display a newly-most-relevant one if applicable.
+	if (bSuccessfullyRemoved && IsValid(SubtitleWidget))
+	{
+		SubtitleWidget->StopDisplayingSubtitle(Subtitle->SubtitleType);
+
+		if (!ActiveSubtitles.IsEmpty())
+		{
+			const USubtitleAssetUserData& HighestPrioritySubtitle = *ActiveSubtitles[0].Subtitle;
+			SubtitleWidget->StartDisplayingSubtitle(HighestPrioritySubtitle);
+		}
+	}
+	else if(!IsValid(SubtitleWidget))
+	{
+		UE_LOG(LogSubtitlesAndClosedCaptions, Warning, TEXT("Can't remove subtitles because there isn't a valid UMG widget."));
+	}
+
+	// Also remove delayed-start subtitles using this asset.
+	FirstRemovedIndex = Algo::RemoveIf(DelayedSubtitles, [&Subtitle](const FActiveSubtitle& DelayedSubtitle) {return DelayedSubtitle.Subtitle == Subtitle; });
+
+	for (int32 Index = FirstRemovedIndex; Index < DelayedSubtitles.Num(); ++Index)
+	{
+		TimerManager.ClearTimer(DelayedSubtitles[Index].DurationTimerHandle);
+	}
+	DelayedSubtitles.SetNum(FirstRemovedIndex);
+}
+
+bool USubtitlesSubsystem::TryCreateUMGWidget()
+{
+	// Set up the UMG widget
+	const USubtitlesSettings* Settings = GetDefault<USubtitlesSettings>();
+	check(Settings != nullptr);
+	const TSubclassOf<USubtitleTextBlock>& WidgetToUse = Settings->GetWidget();
+	if (IsValid(WidgetToUse))
+	{
+		SubtitleWidget = CreateWidget<USubtitleTextBlock>(GetWorld(), WidgetToUse);
+	}
+	else
+	{
+		// Fallback to default widget (not set by user):
+		const TSubclassOf<USubtitleTextBlock>& WidgetToUseDefault = Settings->GetWidgetDefault();
+		if (IsValid(WidgetToUseDefault))
+		{
+			SubtitleWidget = CreateWidget<USubtitleTextBlock>(GetWorld(), WidgetToUseDefault);
+		}
+	}
+	bInitializedWidget = false;
+
+	return IsValid(SubtitleWidget);
+}
+
+void USubtitlesSubsystem::UpdateWidgetData()
+{
+	// Update the widget. If it's not valid (eg, destroyed on non-seamless travel), try re-creating it first.
+	if (IsValid(SubtitleWidget) || (bInitializedWidget && TryCreateUMGWidget() && IsValid(SubtitleWidget)))
+	{
+		if (!bInitializedWidget)
+		{
+			SubtitleWidget->AddToViewport();
+			SubtitleWidget->SetVisibility(ESlateVisibility::HitTestInvisible);
+			bInitializedWidget = true;
+		}
+
+		const USubtitleAssetUserData& NewHighestPrioritySubtitle = *ActiveSubtitles[0].Subtitle;
+		SubtitleWidget->StartDisplayingSubtitle(NewHighestPrioritySubtitle);
+	}
+	else
+	{
+		UE_LOG(LogSubtitlesAndClosedCaptions, Warning, TEXT("Can't display subtitles because there isn't a valid UMG widget to display it to (check Project Settings)."));
+	}
+}
+
diff --git a/Engine/Plugins/Experimental/SubtitlesAndClosedCaptions/Source/SubtitlesAndClosedCaptions/Public/ActiveSubtitle.h b/Engine/Plugins/Experimental/SubtitlesAndClosedCaptions/Source/SubtitlesAndClosedCaptions/Public/ActiveSubtitle.h
new file mode 100644
--- /dev/null
+++ b/Engine/Plugins/Experimental/SubtitlesAndClosedCaptions/Source/SubtitlesAndClosedCaptions/Public/ActiveSubtitle.h
@@ -0,0 +1,21 @@
+// Copyright Epic Games, Inc. All Rights Reserved.
+
+#pragma once 
+
+class USubtitleAssetUserData;
+
+#include "Engine/TimerHandle.h"
+#include "UObject/ObjectPtr.h"
+
+#include "ActiveSubtitle.generated.h"
+
+USTRUCT()
+struct FActiveSubtitle
+{
+	GENERATED_BODY()
+
+	UPROPERTY()
+	TObjectPtr<const USubtitleAssetUserData> Subtitle = nullptr;
+
+	FTimerHandle DurationTimerHandle;
+};
diff --git a/Engine/Plugins/Experimental/SubtitlesAndClosedCaptions/Source/SubtitlesAndClosedCaptions/Public/MovieSceneSubtitleComponentTypes.h b/Engine/Plugins/Experimental/SubtitlesAndClosedCaptions/Source/SubtitlesAndClosedCaptions/Public/MovieSceneSubtitleComponentTypes.h
new file mode 100644
--- /dev/null
+++ b/Engine/Plugins/Experimental/SubtitlesAndClosedCaptions/Source/SubtitlesAndClosedCaptions/Public/MovieSceneSubtitleComponentTypes.h
@@ -0,0 +1,26 @@
+// Copyright Epic Games, Inc. All Rights Reserved.
+
+#pragma once
+
+#include "EntitySystem/MovieSceneEntityIDs.h"
+#include "MovieSceneFwd.h"
+#include "SubtitleDataComponent.h"
+
+struct FSubtitleDataComponent;
+
+class UMovieSceneSection;
+
+// Component data present on all Subtitle entities
+struct FMovieSceneSubtitleComponentTypes
+{
+	UE::MovieScene::TComponentTypeID<FSubtitleDataComponent> SubtitleData;
+
+	SUBTITLESANDCLOSEDCAPTIONS_API static FMovieSceneSubtitleComponentTypes* Get();
+	static void Destroy();
+
+private:
+	FMovieSceneSubtitleComponentTypes();
+
+	static bool GMovieSceneSubtitlesComponentTypesDestroyed;
+	static TUniquePtr<FMovieSceneSubtitleComponentTypes> GMovieSceneSubtitlesComponentTypes;
+};
diff --git a/Engine/Plugins/Experimental/SubtitlesAndClosedCaptions/Source/SubtitlesAndClosedCaptions/Public/MovieSceneSubtitleSection.h b/Engine/Plugins/Experimental/SubtitlesAndClosedCaptions/Source/SubtitlesAndClosedCaptions/Public/MovieSceneSubtitleSection.h
new file mode 100644
--- /dev/null
+++ b/Engine/Plugins/Experimental/SubtitlesAndClosedCaptions/Source/SubtitlesAndClosedCaptions/Public/MovieSceneSubtitleSection.h
@@ -0,0 +1,46 @@
+// Copyright Epic Games, Inc. All Rights Reserved.
+
+#pragma once
+
+#include "EntitySystem/IMovieSceneEntityProvider.h"
+#include "EntitySystem/MovieSceneEntityIDs.h"
+#include "MovieSceneSection.h"
+#include "UObject/ObjectMacros.h"
+
+#include "MovieSceneSubtitleSection.generated.h"
+
+class USubtitleAssetUserData;
+struct FMovieSceneEntityComponentFieldBuilder;
+struct FMovieSceneEvaluationFieldEntityMetaData;
+
+UCLASS(MinimalAPI)
+class UMovieSceneSubtitleSection
+	: public UMovieSceneSection
+	, public IMovieSceneEntityProvider
+{
+	GENERATED_UCLASS_BODY()
+
+public:
+
+	SUBTITLESANDCLOSEDCAPTIONS_API TObjectPtr<const USubtitleAssetUserData> GetSubtitle() const
+	{
+		return Subtitle;
+	}
+
+	SUBTITLESANDCLOSEDCAPTIONS_API void SetSubtitle(const USubtitleAssetUserData& InSubtitle)
+	{
+		Subtitle = &InSubtitle;
+	}
+
+	//~ UMovieSceneSection interface
+	SUBTITLESANDCLOSEDCAPTIONS_API virtual TOptional<TRange<FFrameNumber> > GetAutoSizeRange() const override;
+
+	//~ IMovieSceneEntityProvider interface
+	SUBTITLESANDCLOSEDCAPTIONS_API virtual bool PopulateEvaluationFieldImpl(const TRange<FFrameNumber>& EffectiveRange, const FMovieSceneEvaluationFieldEntityMetaData& InMetaData, FMovieSceneEntityComponentFieldBuilder* OutFieldBuilder) override;
+	SUBTITLESANDCLOSEDCAPTIONS_API virtual void ImportEntityImpl(UMovieSceneEntitySystemLinker* EntityLinker, const FEntityImportParams& Params, FImportedEntity* OutImportedEntity) override;
+
+private:
+
+	UPROPERTY(EditAnywhere, Category = "Subtitles")
+	TObjectPtr<const USubtitleAssetUserData> Subtitle;
+};
diff --git a/Engine/Plugins/Experimental/SubtitlesAndClosedCaptions/Source/SubtitlesAndClosedCaptions/Public/MovieSceneSubtitlesSystem.h b/Engine/Plugins/Experimental/SubtitlesAndClosedCaptions/Source/SubtitlesAndClosedCaptions/Public/MovieSceneSubtitlesSystem.h
new file mode 100644
--- /dev/null
+++ b/Engine/Plugins/Experimental/SubtitlesAndClosedCaptions/Source/SubtitlesAndClosedCaptions/Public/MovieSceneSubtitlesSystem.h
@@ -0,0 +1,131 @@
+// Copyright Epic Games, Inc. All Rights Reserved.
+
+#pragma once
+
+#include "CoreTypes.h"
+#include "EntitySystem/MovieSceneEntityIDs.h"
+#include "EntitySystem/MovieSceneEntitySystem.h"
+#include "EntitySystem/MovieSceneSequenceInstanceHandle.h"
+#include "MovieSceneTracksComponentTypes.h"
+#include "UObject/ObjectMacros.h"
+
+#include "EntitySystem/MovieSceneEntitySystemLinker.h"
+#include "EntitySystem/MovieSceneEntitySystemTask.h"
+#include "MovieSceneSubtitleSection.h"
+#include "SubtitleDataComponent.h"
+#include "Subtitles/SubtitlesAndClosedCaptionsDelegates.h"
+
+#include "MovieSceneSubtitlesSystem.generated.h"
+
+struct SUBTITLESANDCLOSEDCAPTIONS_API FEvaluateSubtitles
+{
+
+	FEvaluateSubtitles(const UE::MovieScene::FInstanceRegistry& InInstanceRegistry)
+		: InstanceRegistry(InInstanceRegistry)
+	{
+	}
+
+	void ForEachAllocation(
+		const UE::MovieScene::FEntityAllocation* Allocation,
+		UE::MovieScene::TRead<UE::MovieScene::FInstanceHandle> SequenceInstanceHandles,
+		UE::MovieScene::TWrite<FSubtitleDataComponent> SubtitleData) const
+	{
+		const UE::MovieScene::FBuiltInComponentTypes* BuiltInComponents = UE::MovieScene::FBuiltInComponentTypes::Get();
+
+		const int32 Num = Allocation->Num();
+		for (int32 Index = 0; Index < Num; ++Index)
+		{
+			const UE::MovieScene::FInstanceHandle& SequenceInstanceHandle = SequenceInstanceHandles[Index];
+			const UE::MovieScene::FSequenceInstance& SequenceInstance = InstanceRegistry.GetInstance(SequenceInstanceHandle);
+
+			// Evaluate the subtitle associated with each subtitle section with logic dependent on changes in the context playing state
+			Evaluate(SequenceInstance,
+				SubtitleData[Index],
+				SubtitleData[Index].LastSequenceInstanceStatus);
+
+			SubtitleData[Index].LastSequenceInstanceStatus = SequenceInstance.GetContext().GetStatus();
+		}
+	}
+
+	void Evaluate(const UE::MovieScene::FSequenceInstance& SequenceInstance, const FSubtitleDataComponent& SubtitleData, EMovieScenePlayerStatus::Type LastStatus) const
+	{
+		const FMovieSceneContext& Context = SequenceInstance.GetContext();
+
+		const UMovieSceneSubtitleSection* SubtitleSection = SubtitleData.SubtitleSection;
+		if (!ensureMsgf(SubtitleSection, TEXT("No valid subtitle section found in subtitles data!")))
+		{
+			return;
+		}
+		const USubtitleAssetUserData* Subtitle = SubtitleData.SubtitleSection->GetSubtitle();
+		const UAssetUserData& AssetUserData = *CastChecked<const UAssetUserData>(Subtitle);
+		if (!ensureMsgf(Subtitle, TEXT("No valid subtitle found in subtitle section!")))
+		{
+			return;
+		}
+
+		const EMovieScenePlayerStatus::Type Status = Context.GetStatus();
+#if !NO_LOGGING
+		if (const UEnum* Enum = StaticEnum<EMovieScenePlayerStatus::Type>())
+		{
+			/* TODO: enable for ClangEditor Win64
+			const FString LastStatusString = Enum->GetNameStringByIndex(LastStatus);
+			const FString StatusString = Enum->GetNameStringByIndex(Status);
+			UE_LOG(LogSubtitlesAndClosedCaptions, Verbose, TEXT("LastStatus: %s, Status: %s, HasJumped: %u"),
+				*LastStatusString,
+				*StatusString,
+				!!Context.HasJumped());
+			*/
+		}
+#endif // !NO_LOGGING
+
+		// Stop all subtitles when we start playing to clear out any infinite duration subtitiles that were queued on pause/stop
+		const bool bHasStarted{ (Status == EMovieScenePlayerStatus::Type::Playing || Status == EMovieScenePlayerStatus::Type::Scrubbing)
+		&& (LastStatus == EMovieScenePlayerStatus::Type::Stopped || LastStatus == EMovieScenePlayerStatus::Type::Paused || LastStatus == EMovieScenePlayerStatus::Type::Stepping) };
+		if (bHasStarted || Context.HasJumped())
+		{
+			FSubtitlesAndClosedCaptionsDelegates::StopAllSubtitles.ExecuteIfBound();
+		}
+
+		// Only queue up the subtitle if it isn't already active so we don't spam queueing of subtitles every frame
+		if (FSubtitlesAndClosedCaptionsDelegates::IsSubtitleActive.IsBound())
+		{
+			const FFrameNumber LastSubtitleFrame = SubtitleSection->GetRange().GetUpperBoundValue();
+			const FFrameNumber FirstSubtitleFrame = SubtitleSection->GetRange().GetLowerBoundValue();
+			const FFrameNumber LastEvaluatedFrame = Context.GetFrameNumberRange().GetUpperBoundValue();
+
+			if (!FSubtitlesAndClosedCaptionsDelegates::IsSubtitleActive.Execute(AssetUserData))
+			{
+				// If not already active, queue the subtitle. It will remain visible until manually stopped; if the frame is frozen on a subtitle then display that subtitle forever.
+				if (LastEvaluatedFrame >= FirstSubtitleFrame && LastEvaluatedFrame < LastSubtitleFrame)
+				{
+					FQueueSubtitleParameters Params{ AssetUserData };
+					FSubtitlesAndClosedCaptionsDelegates::QueueSubtitle.ExecuteIfBound(Params, ESubtitleTiming::ExternallyTimed);
+				}
+			}
+			else if(LastSubtitleFrame <= LastEvaluatedFrame)
+			{
+				// If the subtitle is already active and outside the range of frames it should display for, remove it.
+				FSubtitlesAndClosedCaptionsDelegates::StopSubtitle.ExecuteIfBound(AssetUserData);
+			}
+		}
+	}
+
+protected:
+	const UE::MovieScene::FInstanceRegistry& InstanceRegistry;
+};
+
+UCLASS()
+class UMovieSceneSubtitlesSystem : public UMovieSceneEntitySystem
+{
+	GENERATED_BODY()
+
+public:
+
+	UMovieSceneSubtitlesSystem(const FObjectInitializer& ObjInit);
+	virtual ~UMovieSceneSubtitlesSystem();
+
+	//~ UMovieSceneEntitySystem members
+	virtual void OnSchedulePersistentTasks(UE::MovieScene::IEntitySystemScheduler* TaskScheduler) override;
+	virtual void OnRun(FSystemTaskPrerequisites& InPrerequisites, FSystemSubsequentTasks& Subsequents) override;
+};
+
diff --git a/Engine/Plugins/Experimental/SubtitlesAndClosedCaptions/Source/SubtitlesAndClosedCaptions/Public/MovieSceneSubtitlesTrack.h b/Engine/Plugins/Experimental/SubtitlesAndClosedCaptions/Source/SubtitlesAndClosedCaptions/Public/MovieSceneSubtitlesTrack.h
new file mode 100644
--- /dev/null
+++ b/Engine/Plugins/Experimental/SubtitlesAndClosedCaptions/Source/SubtitlesAndClosedCaptions/Public/MovieSceneSubtitlesTrack.h
@@ -0,0 +1,41 @@
+// Copyright Epic Games, Inc. All Rights Reserved.
+
+#pragma once
+
+#include "UObject/ObjectMacros.h"
+#include "MovieSceneNameableTrack.h"
+
+#include "MovieSceneSubtitlesTrack.generated.h"
+
+class UMovieSceneSection;
+class USubtitleAssetUserData;
+
+UCLASS(MinimalAPI)
+class UMovieSceneSubtitlesTrack
+	: public UMovieSceneNameableTrack
+{
+	GENERATED_UCLASS_BODY()
+
+public:
+	SUBTITLESANDCLOSEDCAPTIONS_API UMovieSceneSection* AddNewSubtitle(const USubtitleAssetUserData& Subtitle, FFrameNumber InStartTime);
+	SUBTITLESANDCLOSEDCAPTIONS_API UMovieSceneSection* AddNewSubtitleOnRow(const USubtitleAssetUserData& Subtitle, FFrameNumber InStartTime, int32 InRowIndex);
+
+public:
+
+	// UMovieSceneTrack interface
+	SUBTITLESANDCLOSEDCAPTIONS_API virtual bool SupportsType(TSubclassOf<UMovieSceneSection> SectionClass) const override;
+	SUBTITLESANDCLOSEDCAPTIONS_API virtual void RemoveAllAnimationData() override;
+	SUBTITLESANDCLOSEDCAPTIONS_API virtual bool HasSection(const UMovieSceneSection& Section) const override;
+	SUBTITLESANDCLOSEDCAPTIONS_API virtual void AddSection(UMovieSceneSection& Section) override;
+	SUBTITLESANDCLOSEDCAPTIONS_API virtual void RemoveSection(UMovieSceneSection& Section) override;
+	SUBTITLESANDCLOSEDCAPTIONS_API virtual void RemoveSectionAt(int32 SectionIndex) override;
+	SUBTITLESANDCLOSEDCAPTIONS_API virtual bool IsEmpty() const override;
+	SUBTITLESANDCLOSEDCAPTIONS_API virtual const TArray<UMovieSceneSection*>& GetAllSections() const override;
+	SUBTITLESANDCLOSEDCAPTIONS_API virtual bool SupportsMultipleRows() const override;
+	SUBTITLESANDCLOSEDCAPTIONS_API virtual UMovieSceneSection* CreateNewSection() override;
+
+private:
+	UPROPERTY()
+	TArray<TObjectPtr<UMovieSceneSection>> SubtitleSections;
+
+};
diff --git a/Engine/Plugins/Experimental/SubtitlesAndClosedCaptions/Source/SubtitlesAndClosedCaptions/Public/SubtitleDataComponent.h b/Engine/Plugins/Experimental/SubtitlesAndClosedCaptions/Source/SubtitlesAndClosedCaptions/Public/SubtitleDataComponent.h
new file mode 100644
--- /dev/null
+++ b/Engine/Plugins/Experimental/SubtitlesAndClosedCaptions/Source/SubtitlesAndClosedCaptions/Public/SubtitleDataComponent.h
@@ -0,0 +1,23 @@
+// Copyright Epic Games, Inc. All Rights Reserved.
+
+#pragma once
+
+#include "EntitySystem/MovieSceneEntityIDs.h"
+#include "MovieSceneFwd.h"
+#include "UObject/ObjectPtr.h"
+
+#include "SubtitleDataComponent.generated.h"
+
+class UMovieSceneSubtitleSection;
+
+/** Component data for subtitles tracks */
+USTRUCT()
+struct SUBTITLESANDCLOSEDCAPTIONS_API FSubtitleDataComponent
+{
+	GENERATED_BODY()
+
+	UPROPERTY()
+	TObjectPtr<const UMovieSceneSubtitleSection> SubtitleSection;
+
+	EMovieScenePlayerStatus::Type LastSequenceInstanceStatus = EMovieScenePlayerStatus::Type::Stopped;
+};
diff --git a/Engine/Plugins/Experimental/SubtitlesAndClosedCaptions/Source/SubtitlesAndClosedCaptions/Public/SubtitleTextBlock.h b/Engine/Plugins/Experimental/SubtitlesAndClosedCaptions/Source/SubtitlesAndClosedCaptions/Public/SubtitleTextBlock.h
new file mode 100644
--- /dev/null
+++ b/Engine/Plugins/Experimental/SubtitlesAndClosedCaptions/Source/SubtitlesAndClosedCaptions/Public/SubtitleTextBlock.h
@@ -0,0 +1,32 @@
+// Copyright Epic Games, Inc. All Rights Reserved.
+
+#pragma once
+#include "Blueprint/UserWidget.h"
+#include "Components/TextBlock.h"
+
+#include "Subtitles/SubtitlesAndClosedCaptionsDelegates.h"
+
+#include "SubtitleTextBlock.generated.h"
+
+UCLASS(Abstract, BlueprintType, Blueprintable, meta=( DontUseGenericSpawnObject="True", DisableNativeTick), MinimalAPI)
+class USubtitleTextBlock : public UUserWidget
+{
+	GENERATED_BODY()
+
+protected:
+	virtual void NativeConstruct() override;
+
+private:
+	UPROPERTY(EditDefaultsOnly, Category = "Subtitles", meta = (BindWidget))
+	TObjectPtr<UTextBlock> DialogSubtitleBlock;
+
+	UPROPERTY(EditDefaultsOnly, Category = "Subtitles", meta = (BindWidget))
+	TObjectPtr<UTextBlock> CaptionSubtitleBlock;
+
+	UPROPERTY(EditDefaultsOnly, Category = "Subtitles", meta = (BindWidget))
+	TObjectPtr<UTextBlock> DescriptionSubtitleBlock;
+
+public:
+	void StartDisplayingSubtitle(const USubtitleAssetUserData& Subtitle);
+	void StopDisplayingSubtitle(const ESubtitleType SubtitleType);
+};
\ No newline at end of file
diff --git a/Engine/Plugins/Experimental/SubtitlesAndClosedCaptions/Source/SubtitlesAndClosedCaptions/Public/SubtitlesAndClosedCaptionsModule.h b/Engine/Plugins/Experimental/SubtitlesAndClosedCaptions/Source/SubtitlesAndClosedCaptions/Public/SubtitlesAndClosedCaptionsModule.h
new file mode 100644
--- /dev/null
+++ b/Engine/Plugins/Experimental/SubtitlesAndClosedCaptions/Source/SubtitlesAndClosedCaptions/Public/SubtitlesAndClosedCaptionsModule.h
@@ -0,0 +1,19 @@
+// Copyright Epic Games, Inc. All Rights Reserved.
+
+#pragma once
+
+#include "Modules/ModuleInterface.h"
+#include "Subtitles/ISubtitlesAndClosedCaptionsModule.h"
+#include "Templates/SubclassOf.h"
+
+SUBTITLESANDCLOSEDCAPTIONS_API DECLARE_LOG_CATEGORY_EXTERN(LogSubtitlesAndClosedCaptions, Log, All);
+
+class UAssetUserData;
+
+class FSubtitlesAndClosedCaptionsModule : public ISubtitlesAndClosedCaptionsModule
+{
+public:
+	void StartupModule() override;
+	void ShutdownModule() override;
+	TSubclassOf<UAssetUserData> GetAssetUserDataClass() const override;
+};
diff --git a/Engine/Plugins/Experimental/SubtitlesAndClosedCaptions/Source/SubtitlesAndClosedCaptions/Public/SubtitlesBlueprintFunctionLibrary.h b/Engine/Plugins/Experimental/SubtitlesAndClosedCaptions/Source/SubtitlesAndClosedCaptions/Public/SubtitlesBlueprintFunctionLibrary.h
new file mode 100644
--- /dev/null
+++ b/Engine/Plugins/Experimental/SubtitlesAndClosedCaptions/Source/SubtitlesAndClosedCaptions/Public/SubtitlesBlueprintFunctionLibrary.h
@@ -0,0 +1,29 @@
+// Copyright Epic Games, Inc. All Rights Reserved.
+
+#pragma once
+
+#include "Kismet/BlueprintFunctionLibrary.h"
+#include "Subtitles/SubtitlesAndClosedCaptionsDelegates.h"
+#include "SubtitlesBlueprintFunctionLibrary.generated.h"
+
+class USubtitleAssetUserData;
+
+UCLASS(meta = (ScriptName = "SubtitlesLibrary"))
+class SUBTITLESANDCLOSEDCAPTIONS_API USubtitlesBlueprintFunctionLibrary : public UBlueprintFunctionLibrary
+{
+	GENERATED_BODY()
+
+public:
+
+	UFUNCTION(BlueprintCallable, Category = Subtitles)
+	static void QueueSubtitle(const USubtitleAssetUserData* Subtitle, const ESubtitleTiming Timing = ESubtitleTiming::InternallyTimed);
+
+	UFUNCTION(BlueprintCallable, Category = Subtitles)
+	static bool IsSubtitleActive(const USubtitleAssetUserData* Subtitle);
+
+	UFUNCTION(BlueprintCallable, Category = Subtitles)
+	static void StopSubtitle(const USubtitleAssetUserData* Subtitle);
+
+	UFUNCTION(BlueprintCallable, Category = Subtitles)
+	static void StopAllSubtitles();
+};
\ No newline at end of file
diff --git a/Engine/Plugins/Experimental/SubtitlesAndClosedCaptions/Source/SubtitlesAndClosedCaptions/Public/SubtitlesSettings.h b/Engine/Plugins/Experimental/SubtitlesAndClosedCaptions/Source/SubtitlesAndClosedCaptions/Public/SubtitlesSettings.h
new file mode 100644
--- /dev/null
+++ b/Engine/Plugins/Experimental/SubtitlesAndClosedCaptions/Source/SubtitlesAndClosedCaptions/Public/SubtitlesSettings.h
@@ -0,0 +1,29 @@
+// Copyright Epic Games, Inc. All Rights Reserved.
+
+#pragma once
+
+#include "Engine/DeveloperSettings.h"
+#include "Math/Color.h"
+#include "SubtitleTextBlock.h"
+#include "SubtitlesSettings.generated.h"
+
+
+UCLASS(config = Engine, defaultconfig, meta = (DisplayName = "Subtitles And Closed Captions"))
+class USubtitlesSettings : public UDeveloperSettings
+{
+	GENERATED_BODY()
+
+public:
+	USubtitlesSettings();
+
+	const TSubclassOf<USubtitleTextBlock>& GetWidget() const { return SubtitleWidgetToUse; }
+	const TSubclassOf<USubtitleTextBlock>& GetWidgetDefault() const { return SubtitleWidgetToUseDefault; }
+
+protected:
+
+	UPROPERTY(config, EditDefaultsOnly, Category = Subtitles, AdvancedDisplay)
+	TSubclassOf<USubtitleTextBlock> SubtitleWidgetToUse;
+
+	UPROPERTY(config)
+	TSubclassOf<USubtitleTextBlock> SubtitleWidgetToUseDefault; // fallback for SubtitleWidgetToUse (not set by user)
+};
diff --git a/Engine/Plugins/Experimental/SubtitlesAndClosedCaptions/Source/SubtitlesAndClosedCaptions/Public/SubtitlesSubsystem.h b/Engine/Plugins/Experimental/SubtitlesAndClosedCaptions/Source/SubtitlesAndClosedCaptions/Public/SubtitlesSubsystem.h
new file mode 100644
--- /dev/null
+++ b/Engine/Plugins/Experimental/SubtitlesAndClosedCaptions/Source/SubtitlesAndClosedCaptions/Public/SubtitlesSubsystem.h
@@ -0,0 +1,111 @@
+// Copyright Epic Games, Inc. All Rights Reserved.
+
+#pragma once 
+
+#include "ActiveSubtitle.h"
+#include "Fonts/SlateFontInfo.h"
+#include "Math/Color.h"
+#include "Subsystems/WorldSubsystem.h"
+
+#include "SubtitleTextBlock.h"
+
+#include "Subtitles/SubtitlesAndClosedCaptionsDelegates.h"
+#include "SubtitlesSubsystem.generated.h"
+
+class FCanvas;
+class UAssetUserData;
+class UFont;
+struct FQueueSubtitleParameters;
+
+#if WITH_DEV_AUTOMATION_TESTS
+namespace SubtitlesAndClosedCaptions::Test
+{
+	struct FMovieSceneSubtitlesTest;
+	struct FSubtitlesTest;
+}
+#endif
+
+/*
+* #SUBTITLES_PRD -	Requirement:	Ability to allow designers to “script” subtitle location for sequences and scenes to avoid subtitles overlapping important scenes or characters
+*					Use a UEngineSubsystem for blueprints
+*
+*	Game configuration for font customization per game
+*/
+UCLASS(config = Game, defaultConfig)
+class SUBTITLESANDCLOSEDCAPTIONS_API USubtitlesSubsystem : public UWorldSubsystem
+{
+	GENERATED_BODY()
+public:
+	USubtitlesSubsystem() = default;
+
+	// FSubtitlesAndClosedCaptionsDelegates
+
+	// Adds a subtitle to the queue: Params contains the subtitle asset and an optional duration. The highest-priority subtitle in the queue will be displayed.
+	// If Timing is ExternallyTimed, the queued subtitle will remain in the queue until manually removed.
+	// If the subtitle asset has a non-zero StartOffset, it will sit in a delayed-start queue instead of being queued for display.
+	virtual void QueueSubtitle(const FQueueSubtitleParameters& Params, const ESubtitleTiming Timing = ESubtitleTiming::InternallyTimed);
+
+	// Returns true if the given subtitle asset is being displayed.
+	virtual bool IsSubtitleActive(const UAssetUserData& Data) const;
+
+	// Stops the given subtitle asset being displayed.  This includes subtitles not yet being displayed due to their StartOffset
+	virtual void StopSubtitle(const UAssetUserData& Data);
+
+	// Stops all queued subtitles from being displayed.  This includes subtitles not yet being displayed due to their StartOffset.
+	virtual void StopAllSubtitles();
+
+protected:
+
+	virtual void Initialize(FSubsystemCollectionBase& Collection) override;
+	virtual void BindDelegates();
+
+	UPROPERTY()
+	TMap<ESubtitleType, FSlateFontInfo> SubtitleFontInfo;
+
+	// Sorted by priority, desc
+	UPROPERTY()
+	TArray<FActiveSubtitle> ActiveSubtitles;
+
+	// Unsorted; subtitles with a delayed start offset still need to be tracked before entering the queue proper.
+	UPROPERTY()
+	TArray<FActiveSubtitle> DelayedSubtitles;
+
+protected:
+	virtual void AddActiveSubtitle(const USubtitleAssetUserData& Subtitle, float Duration);
+	virtual void MakeDelayedSubtitleActive(TObjectPtr<const USubtitleAssetUserData> Subtitle);
+	virtual void RemoveActiveSubtitle(TObjectPtr<const USubtitleAssetUserData> Subtitle);
+
+#if WITH_DEV_AUTOMATION_TESTS
+	friend struct SubtitlesAndClosedCaptions::Test::FMovieSceneSubtitlesTest;
+	friend struct SubtitlesAndClosedCaptions::Test::FSubtitlesTest;
+public:
+	const TArray<FActiveSubtitle>& GetActiveSubtitles() const
+	{
+		return ActiveSubtitles;
+	}
+
+	// As this is for testing, assume that the number of subtitles has already been checked.
+	const USubtitleAssetUserData* GetTopRankedSubtitle() const
+	{
+		check(ActiveSubtitles.Num() > 0);
+		const FActiveSubtitle& ActiveSubtitle = ActiveSubtitles[0];
+
+		return ActiveSubtitle.Subtitle.Get();
+	}
+
+	void TestActivatingDelayedSubtitle(const UAssetUserData& Data)
+	{
+		const USubtitleAssetUserData* Subtitle = CastChecked<const USubtitleAssetUserData>(&Data);
+		MakeDelayedSubtitleActive(Subtitle);
+	}
+#endif
+	//	WITH_DEV_AUTOMATION_TESTS
+
+private:
+	UPROPERTY()
+	TObjectPtr<USubtitleTextBlock> SubtitleWidget = nullptr;
+
+	bool bInitializedWidget = false;
+	bool TryCreateUMGWidget();
+	void UpdateWidgetData();
+};
diff --git a/Engine/Plugins/Experimental/SubtitlesAndClosedCaptions/Source/SubtitlesAndClosedCaptions/SubtitlesAndClosedCaptions.Build.cs b/Engine/Plugins/Experimental/SubtitlesAndClosedCaptions/Source/SubtitlesAndClosedCaptions/SubtitlesAndClosedCaptions.Build.cs
new file mode 100644
--- /dev/null
+++ b/Engine/Plugins/Experimental/SubtitlesAndClosedCaptions/Source/SubtitlesAndClosedCaptions/SubtitlesAndClosedCaptions.Build.cs
@@ -0,0 +1,35 @@
+// Copyright Epic Games, Inc. All Rights Reserved.
+
+namespace UnrealBuildTool.Rules
+{
+	public class SubtitlesAndClosedCaptions : ModuleRules
+	{
+		public SubtitlesAndClosedCaptions(ReadOnlyTargetRules Target) : base(Target)
+		{
+			PCHUsage = PCHUsageMode.UseExplicitOrSharedPCHs;
+
+			PrivateDependencyModuleNames.AddRange
+			(
+				new string[]
+				{
+					"Slate",
+					"SlateCore",
+				}
+			);
+
+			PublicDependencyModuleNames.AddRange(
+				new string[]
+				{
+					"Core",
+					"CoreUObject",
+					"DeveloperSettings",
+					"Engine",
+					"InputCore",
+					"MovieSceneTracks",
+					"MovieScene",
+					"UMG"
+				}
+			);
+		}
+	}
+}
diff --git a/Engine/Plugins/Experimental/SubtitlesAndClosedCaptions/Source/SubtitlesAndClosedCaptionsEditor/Private/SubtitleAssetDefinitions.h b/Engine/Plugins/Experimental/SubtitlesAndClosedCaptions/Source/SubtitlesAndClosedCaptionsEditor/Private/SubtitleAssetDefinitions.h
new file mode 100644
--- /dev/null
+++ b/Engine/Plugins/Experimental/SubtitlesAndClosedCaptions/Source/SubtitlesAndClosedCaptionsEditor/Private/SubtitleAssetDefinitions.h
@@ -0,0 +1,28 @@
+// Copyright Epic Games, Inc. All Rights Reserved.
+
+#pragma once
+
+#include "AssetDefinitionDefault.h"
+#include "Subtitles/SubtitlesAndClosedCaptionsDelegates.h"
+#include "SubtitlesAndClosedCaptionsEditorModule.h"
+
+#include "SubtitleAssetDefinitions.generated.h"
+
+UCLASS()
+class USubtitleAssetDefinitions : public UAssetDefinitionDefault
+{
+	GENERATED_BODY()
+
+public:
+	// UAssetDefinition Begin
+	virtual FText GetAssetDisplayName() const override { return NSLOCTEXT("AssetTypeActions", "AssetTypeActions_Subtitle", "Subtitle"); }
+	virtual FLinearColor GetAssetColor() const override { return FLinearColor(FColor(97, 85, 212)); }
+	virtual TSoftClassPtr<UObject> GetAssetClass() const override { return USubtitleAssetUserData::StaticClass(); }
+	virtual TConstArrayView<FAssetCategoryPath> GetAssetCategories() const override
+	{
+
+		static const FAssetCategoryPath Categories[1] = {FAssetCategoryPath(ISubtitlesAndClosedCaptionsEditorModule::GetAssetTypeCategory())};
+		return Categories;
+	}
+	// UAssetDefinition End
+};
diff --git a/Engine/Plugins/Experimental/SubtitlesAndClosedCaptions/Source/SubtitlesAndClosedCaptionsEditor/Private/SubtitleFactory.cpp b/Engine/Plugins/Experimental/SubtitlesAndClosedCaptions/Source/SubtitlesAndClosedCaptionsEditor/Private/SubtitleFactory.cpp
new file mode 100644
--- /dev/null
+++ b/Engine/Plugins/Experimental/SubtitlesAndClosedCaptions/Source/SubtitlesAndClosedCaptionsEditor/Private/SubtitleFactory.cpp
@@ -0,0 +1,20 @@
+// Copyright Epic Games, Inc. All Rights Reserved.
+#include "SubtitleFactory.h"
+
+#include "Subtitles/SubtitlesAndClosedCaptionsDelegates.h"
+
+#include UE_INLINE_GENERATED_CPP_BY_NAME(SubtitleFactory)
+
+USubtitleFactory::USubtitleFactory(const FObjectInitializer& ObjectInitializer)
+	: Super(ObjectInitializer)
+{
+	bCreateNew = true;
+	bEditorImport = false;
+	bEditAfterNew = true;
+	SupportedClass = USubtitleAssetUserData::StaticClass();
+}
+
+UObject* USubtitleFactory::FactoryCreateNew(UClass* InClass, UObject* InParent, FName InName, EObjectFlags InFlags, UObject* InContext, FFeedbackContext* InFeedbackContext)
+{
+	return NewObject<USubtitleAssetUserData>(InParent, InClass, InName, InFlags);
+}
diff --git a/Engine/Plugins/Experimental/SubtitlesAndClosedCaptions/Source/SubtitlesAndClosedCaptionsEditor/Private/SubtitleFactory.h b/Engine/Plugins/Experimental/SubtitlesAndClosedCaptions/Source/SubtitlesAndClosedCaptionsEditor/Private/SubtitleFactory.h
new file mode 100644
--- /dev/null
+++ b/Engine/Plugins/Experimental/SubtitlesAndClosedCaptions/Source/SubtitlesAndClosedCaptionsEditor/Private/SubtitleFactory.h
@@ -0,0 +1,18 @@
+// Copyright Epic Games, Inc. All Rights Reserved.
+
+#pragma once
+
+#include "Factories/Factory.h"
+#include "UObject/ObjectMacros.h"
+
+#include "SubtitleFactory.generated.h"
+
+UCLASS(hidecategories = Object, MinimalAPI)
+class USubtitleFactory : public UFactory
+{
+	GENERATED_UCLASS_BODY()
+
+	// UFactory Begin
+	virtual UObject* FactoryCreateNew(UClass* InClass, UObject* InParent, FName InName, EObjectFlags InFlags, UObject* InContext, FFeedbackContext* InFeedbackContext) override;
+	// UFactory End
+};
diff --git a/Engine/Plugins/Experimental/SubtitlesAndClosedCaptions/Source/SubtitlesAndClosedCaptionsEditor/Private/SubtitleSequencerSection.cpp b/Engine/Plugins/Experimental/SubtitlesAndClosedCaptions/Source/SubtitlesAndClosedCaptionsEditor/Private/SubtitleSequencerSection.cpp
new file mode 100644
--- /dev/null
+++ b/Engine/Plugins/Experimental/SubtitlesAndClosedCaptions/Source/SubtitlesAndClosedCaptionsEditor/Private/SubtitleSequencerSection.cpp
@@ -0,0 +1,28 @@
+// Copyright Epic Games, Inc. All Rights Reserved.
+
+#include "SubtitleSequencerSection.h"
+
+#include "MovieSceneSection.h"
+#include "MovieSceneSubtitleSection.h"
+#include "SequencerSectionPainter.h"
+#include "Subtitles/SubtitlesAndClosedCaptionsDelegates.h"
+
+FText FSubtitleSequencerSection::GetSectionTitle() const
+{
+	if (TStrongObjectPtr<const UMovieSceneSection> SectionPin = WeakSection.Pin())
+	{
+		const UMovieSceneSubtitleSection& SubtitlesSection = *CastChecked<const UMovieSceneSubtitleSection>(SectionPin.Get());
+		if (const USubtitleAssetUserData* Subtitle = SubtitlesSection.GetSubtitle())
+		{
+			return Subtitle->Text;
+		}
+	}
+
+	return NSLOCTEXT("FSubtitlesSequencerSection", "NoSubtitleName", "No Subtitle");
+}
+
+int32 FSubtitleSequencerSection::OnPaintSection(FSequencerSectionPainter& Painter) const
+{
+	int32 LayerId = Painter.PaintSectionBackground();
+	return LayerId;
+}
diff --git a/Engine/Plugins/Experimental/SubtitlesAndClosedCaptions/Source/SubtitlesAndClosedCaptionsEditor/Private/SubtitleSequencerSection.h b/Engine/Plugins/Experimental/SubtitlesAndClosedCaptions/Source/SubtitlesAndClosedCaptionsEditor/Private/SubtitleSequencerSection.h
new file mode 100644
--- /dev/null
+++ b/Engine/Plugins/Experimental/SubtitlesAndClosedCaptions/Source/SubtitlesAndClosedCaptionsEditor/Private/SubtitleSequencerSection.h
@@ -0,0 +1,23 @@
+// Copyright Epic Games, Inc. All Rights Reserved.
+
+#pragma once
+
+#include "ISequencerSection.h"
+#include "SubtitleSequencerSection.h"
+
+class FSequencerSectionPainter;
+
+class FSubtitleSequencerSection
+	: public FSequencerSection
+	, public TSharedFromThis<FSubtitleSequencerSection>
+{
+public:
+	FSubtitleSequencerSection(UMovieSceneSection& InSection)
+		: FSequencerSection(InSection)
+	{}
+
+	// ISequencerSection interface
+	virtual FText GetSectionTitle() const override;
+	virtual bool SectionIsResizable() const override { return false; }
+	virtual int32 OnPaintSection(FSequencerSectionPainter& Painter) const override;
+};
diff --git a/Engine/Plugins/Experimental/SubtitlesAndClosedCaptions/Source/SubtitlesAndClosedCaptionsEditor/Private/SubtitlesAndClosedCaptionsEditorModule.cpp b/Engine/Plugins/Experimental/SubtitlesAndClosedCaptions/Source/SubtitlesAndClosedCaptionsEditor/Private/SubtitlesAndClosedCaptionsEditorModule.cpp
new file mode 100644
--- /dev/null
+++ b/Engine/Plugins/Experimental/SubtitlesAndClosedCaptions/Source/SubtitlesAndClosedCaptionsEditor/Private/SubtitlesAndClosedCaptionsEditorModule.cpp
@@ -0,0 +1,46 @@
+// Copyright Epic Games, Inc. All Rights Reserved.
+
+#include "SubtitlesAndClosedCaptionsEditorModule.h"
+
+#include "AssetToolsModule.h"
+#include "AssetTypeCategories.h"
+#include "ISequencerModule.h"
+#include "Modules/ModuleManager.h"
+#include "SubtitlesTrackEditor.h"
+
+#define LOCTEXT_NAMESPACE "SubtitlesAndClosedCaptionsEditor"
+
+FText ISubtitlesAndClosedCaptionsEditorModule::AssetTypeCategory = LOCTEXT("SubtitlesAssetTypeCategory", "Subtitles");
+
+class FSubtitlesAndClosedCaptionsEditorModule : public ISubtitlesAndClosedCaptionsEditorModule
+{
+public:
+	virtual void StartupModule() override
+	{
+		FModuleManager::Get().LoadModuleChecked("SubtitlesAndClosedCaptions");
+
+		// Register assets
+		IAssetTools& AssetTools = FModuleManager::LoadModuleChecked<FAssetToolsModule>("AssetTools").Get();
+		AssetTools.RegisterAdvancedAssetCategory(FName(TEXT("Subtitles")), AssetTypeCategory);
+
+		//
+		ISequencerModule& SequencerModule = FModuleManager::LoadModuleChecked<ISequencerModule>("Sequencer");
+		CreateTrackEditorHandle = SequencerModule.RegisterTrackEditor(FOnCreateTrackEditor::CreateStatic(&FSubtitlesTrackEditor::CreateTrackEditor));
+
+	}
+
+	virtual void ShutdownModule() override
+	{
+		ISequencerModule* SequencerModule = FModuleManager::GetModulePtr<ISequencerModule>("Sequencer");
+		if (SequencerModule != nullptr)
+		{
+			SequencerModule->UnRegisterTrackEditor(CreateTrackEditorHandle);
+		}
+	}
+
+	FDelegateHandle CreateTrackEditorHandle;
+};
+
+IMPLEMENT_MODULE(FSubtitlesAndClosedCaptionsEditorModule, SubtitlesAndClosedCaptionsEditor);
+
+#undef LOCTEXT_NAMESPACE
diff --git a/Engine/Plugins/Experimental/SubtitlesAndClosedCaptions/Source/SubtitlesAndClosedCaptionsEditor/Private/SubtitlesTrackEditor.cpp b/Engine/Plugins/Experimental/SubtitlesAndClosedCaptions/Source/SubtitlesAndClosedCaptionsEditor/Private/SubtitlesTrackEditor.cpp
new file mode 100644
--- /dev/null
+++ b/Engine/Plugins/Experimental/SubtitlesAndClosedCaptions/Source/SubtitlesAndClosedCaptionsEditor/Private/SubtitlesTrackEditor.cpp
@@ -0,0 +1,202 @@
+// Copyright Epic Games, Inc. All Rights Reserved.
+
+#include "SubtitlesTrackEditor.h"
+
+#include "AssetRegistry/AssetRegistryModule.h"
+#include "ContentBrowserDelegates.h"
+#include "ContentBrowserModule.h"
+#include "IContentBrowserSingleton.h"
+#include "Input/Reply.h"
+#include "LevelSequence.h"
+#include "MovieSceneSubtitlesTrack.h"
+#include "MVVM/Views/ViewUtilities.h"
+#include "SequencerSettings.h"
+#include "Subtitles/SubtitlesAndClosedCaptionsDelegates.h"
+#include "SubtitleSequencerSection.h"
+
+
+#define LOCTEXT_NAMESPACE "FSubtitlesTrackEditor"
+
+FSubtitlesTrackEditor::FSubtitlesTrackEditor(TSharedRef<ISequencer> InSequencer)
+	: FMovieSceneTrackEditor(InSequencer)
+{
+}
+
+TSharedRef<ISequencerTrackEditor> FSubtitlesTrackEditor::CreateTrackEditor(TSharedRef<ISequencer> OwningSequencer)
+{
+	return MakeShareable(new FSubtitlesTrackEditor(OwningSequencer));
+}
+
+FText FSubtitlesTrackEditor::GetDisplayName() const
+{
+	return LOCTEXT("SubtitilesTrackEditor_DisplayName", "Subtitles");
+}
+
+void FSubtitlesTrackEditor::BuildAddTrackMenu(FMenuBuilder& MenuBuilder)
+{
+	MenuBuilder.AddMenuEntry(
+		LOCTEXT("AddTrack", "Subtitles Track"),
+		LOCTEXT("AddTooltip", "Adds a new subtitles track that can display subtitles and closed captions."),
+		FSlateIcon(FAppStyle::GetAppStyleSetName(), "Sequencer.Tracks.Audio"), // #SUBTITLES_PRD_TODO - Add a subtitles track icon
+		FUIAction(
+			FExecuteAction::CreateRaw(this, &FSubtitlesTrackEditor::HandleAddMenuEntryExecute)
+		)
+	);
+}
+
+void FSubtitlesTrackEditor::HandleAddMenuEntryExecute()
+{
+	UMovieScene* FocusedMovieScene = GetFocusedMovieScene();
+
+	if (FocusedMovieScene == nullptr)
+	{
+		return;
+	}
+
+	if (FocusedMovieScene->IsReadOnly())
+	{
+		return;
+	}
+
+	const FScopedTransaction Transaction(NSLOCTEXT("Sequencer", "AddSubtitlesTrack_Transaction", "Add Subtitles Track"));
+	FocusedMovieScene->Modify();
+
+	UMovieSceneSubtitlesTrack* NewTrack = FocusedMovieScene->AddTrack<UMovieSceneSubtitlesTrack>();
+	check(NewTrack);
+
+	NewTrack->SetDisplayName(LOCTEXT("SubtitlesTrackName", "Subtitles"));
+
+	if (GetSequencer().IsValid())
+	{
+		GetSequencer()->OnAddTrack(NewTrack, FGuid());
+	}
+}
+
+TSharedPtr<SWidget> FSubtitlesTrackEditor::BuildOutlinerEditWidget(const FGuid& ObjectBinding, UMovieSceneTrack* Track, const FBuildEditWidgetParams& Params)
+{
+	TDelegate OnAssetSelected = FOnAssetSelected::CreateRaw(this, &FSubtitlesTrackEditor::OnAssetSelected, Track);
+	TDelegate OnAssetEnterPressed = FOnAssetEnterPressed::CreateRaw(this, &FSubtitlesTrackEditor::OnAssetEnterPressed, Track);
+	return UE::Sequencer::MakeAddButton(
+		LOCTEXT("SubtitleText", "Subtitle")
+		, FOnGetContent::CreateSP(this, &FSubtitlesTrackEditor::BuildSubMenu, MoveTemp(OnAssetSelected), MoveTemp(OnAssetEnterPressed))
+		, Params.ViewModel);
+}
+
+TSharedRef<SWidget> FSubtitlesTrackEditor::BuildSubMenu(FOnAssetSelected OnAssetSelected, FOnAssetEnterPressed OnAssetEnterPressed)
+{
+	TSharedPtr<ISequencer> SequencerPtr = GetSequencer();
+	UMovieSceneSequence* Sequence = SequencerPtr.IsValid() ? SequencerPtr->GetFocusedMovieSceneSequence() : nullptr;
+
+	FAssetRegistryModule& AssetRegistryModule = FModuleManager::LoadModuleChecked<FAssetRegistryModule>(TEXT("AssetRegistry"));
+	TArray<FTopLevelAssetPath> ClassNames;
+	ClassNames.Add(USubtitleAssetUserData::StaticClass()->GetClassPathName());
+	TSet<FTopLevelAssetPath> DerivedClassNames;
+	AssetRegistryModule.Get().GetDerivedClassNames(ClassNames, TSet<FTopLevelAssetPath>(), DerivedClassNames);
+
+	FMenuBuilder MenuBuilder(true, nullptr);
+
+	FAssetPickerConfig AssetPickerConfig;
+	{
+		AssetPickerConfig.OnAssetSelected = OnAssetSelected;
+		AssetPickerConfig.OnAssetEnterPressed = OnAssetEnterPressed;
+		AssetPickerConfig.bAllowNullSelection = false;
+		AssetPickerConfig.bAddFilterUI = true;
+		AssetPickerConfig.bShowTypeInColumnView = false;
+		AssetPickerConfig.InitialAssetViewType = EAssetViewType::List;
+		for (FTopLevelAssetPath ClassName : DerivedClassNames)
+		{
+			AssetPickerConfig.Filter.ClassPaths.Add(ClassName);
+		}
+		AssetPickerConfig.SaveSettingsName = TEXT("SequencerAssetPicker");
+		AssetPickerConfig.AdditionalReferencingAssets.Add(FAssetData(Sequence));
+	}
+
+	FContentBrowserModule& ContentBrowserModule = FModuleManager::Get().LoadModuleChecked<FContentBrowserModule>(TEXT("ContentBrowser"));
+
+	const float WidthOverride = SequencerPtr.IsValid() ? SequencerPtr->GetSequencerSettings()->GetAssetBrowserWidth() : 500.f;
+	const float HeightOverride = SequencerPtr.IsValid() ? SequencerPtr->GetSequencerSettings()->GetAssetBrowserHeight() : 400.f;
+
+	TSharedPtr<SBox> MenuEntry = SNew(SBox)
+		.WidthOverride(WidthOverride)
+		.HeightOverride(HeightOverride)
+		[
+			ContentBrowserModule.Get().CreateAssetPicker(AssetPickerConfig)
+		];
+
+	MenuBuilder.AddWidget(MenuEntry.ToSharedRef(), FText::GetEmpty(), true);
+
+	return MenuBuilder.MakeWidget();
+}
+
+void FSubtitlesTrackEditor::OnAssetSelected(const FAssetData& AssetData, UMovieSceneTrack* Track)
+{
+	FSlateApplication::Get().DismissAllMenus();
+
+	const UObject* SelectedObject = AssetData.GetAsset();
+	if (!SelectedObject)
+	{
+		return;
+	}
+
+	const USubtitleAssetUserData* NewAsset = CastChecked<const USubtitleAssetUserData>(AssetData.GetAsset());
+	if (!NewAsset)
+	{
+		return;
+	}
+	const FScopedTransaction Transaction(NSLOCTEXT("Sequencer", "AddSubtitle_Transaction", "Add Subtitle"));
+
+	UMovieSceneSubtitlesTrack* SubtitlesTrack = Cast<UMovieSceneSubtitlesTrack>(Track);
+	check(SubtitlesTrack);
+	SubtitlesTrack->Modify();
+
+	if (TSharedPtr<ISequencer> SequencerPin = GetSequencer())
+	{
+		UMovieSceneSection* NewSection = SubtitlesTrack->AddNewSubtitle(*NewAsset, SequencerPin->GetLocalTime().Time.FrameNumber);
+
+		SequencerPin->EmptySelection();
+		SequencerPin->SelectSection(NewSection);
+		SequencerPin->ThrobSectionSelection();
+
+		SequencerPin->NotifyMovieSceneDataChanged(EMovieSceneDataChangeType::MovieSceneStructureItemAdded);
+	}
+}
+
+void FSubtitlesTrackEditor::OnAssetEnterPressed(const TArray<FAssetData>& AssetData, UMovieSceneTrack* Track)
+{
+	if (!AssetData.IsEmpty())
+	{
+		OnAssetSelected(AssetData[0].GetAsset(), Track);
+	}
+}
+
+TSharedRef<ISequencerSection> FSubtitlesTrackEditor::MakeSectionInterface(UMovieSceneSection& SectionObject, UMovieSceneTrack& Track, FGuid ObjectBinding)
+{
+	check(SupportsType(SectionObject.GetOuter()->GetClass()));
+	return MakeShareable(new FSubtitleSequencerSection(SectionObject));
+}
+
+bool FSubtitlesTrackEditor::SupportsType(TSubclassOf<UMovieSceneTrack> Type) const
+{
+	return Type == UMovieSceneSubtitlesTrack::StaticClass();
+}
+
+bool FSubtitlesTrackEditor::SupportsSequence(UMovieSceneSequence* InSequence) const
+{
+	ETrackSupport TrackSupported = ETrackSupport::NotSupported;
+	if (InSequence && InSequence->IsA(ULevelSequence::StaticClass()))
+	{
+		TrackSupported = ETrackSupport::Supported;
+	}
+	return TrackSupported == ETrackSupport::Supported;
+}
+
+bool FSubtitlesTrackEditor::OnAllowDrop(const FDragDropEvent& DragDropEvent, FSequencerDragDropParams& DragDropParams)
+{
+	return false;
+}
+
+FReply FSubtitlesTrackEditor::OnDrop(const FDragDropEvent& DragDropEvent, const FSequencerDragDropParams& DragDropParams)
+{
+	return FReply::Unhandled();
+}
+#undef LOCTEXT_NAMESPACE
diff --git a/Engine/Plugins/Experimental/SubtitlesAndClosedCaptions/Source/SubtitlesAndClosedCaptionsEditor/Private/SubtitlesTrackEditor.h b/Engine/Plugins/Experimental/SubtitlesAndClosedCaptions/Source/SubtitlesAndClosedCaptionsEditor/Private/SubtitlesTrackEditor.h
new file mode 100644
--- /dev/null
+++ b/Engine/Plugins/Experimental/SubtitlesAndClosedCaptions/Source/SubtitlesAndClosedCaptionsEditor/Private/SubtitlesTrackEditor.h
@@ -0,0 +1,55 @@
+// Copyright Epic Games, Inc. All Rights Reserved.
+
+#pragma once
+
+#include "MovieSceneTrackEditor.h"
+
+#include "ContentBrowserDelegates.h"
+
+struct FGuid;
+class FReply;
+class UMovieSceneTrack;
+class UMovieSceneSubtitlesTrack;
+
+/**
+ * Tools for subtitles tracks
+ */
+class FSubtitlesTrackEditor
+	: public FMovieSceneTrackEditor
+{
+public:
+
+	/**
+	 * Constructor
+	 *
+	 * @param InSequencer The sequencer instance to be used by this tool
+	 */
+	FSubtitlesTrackEditor(TSharedRef<ISequencer> InSequencer);
+
+	/**
+	 * Creates an instance of this class.  Called by a sequencer
+	 *
+	 * @param OwningSequencer The sequencer instance to be used by this tool
+	 * @return The new instance of this class
+	 */
+	static TSharedRef<ISequencerTrackEditor> CreateTrackEditor(TSharedRef<ISequencer> OwningSequencer);
+
+public:
+
+	// ISequencerTrackEditor interface
+	virtual FText GetDisplayName() const override;
+	virtual void BuildAddTrackMenu(FMenuBuilder& MenuBuilder) override;
+	virtual TSharedPtr<SWidget> BuildOutlinerEditWidget(const FGuid& ObjectBinding, UMovieSceneTrack* Track, const FBuildEditWidgetParams& Params) override;
+	virtual TSharedRef<ISequencerSection> MakeSectionInterface(UMovieSceneSection& SectionObject, UMovieSceneTrack& Track, FGuid ObjectBinding) override;
+	virtual bool SupportsType(TSubclassOf<UMovieSceneTrack> Type) const override;
+	virtual bool SupportsSequence(UMovieSceneSequence* InSequence) const override;
+	virtual bool OnAllowDrop(const FDragDropEvent& DragDropEvent, FSequencerDragDropParams& DragDropParams) override;
+	virtual FReply OnDrop(const FDragDropEvent& DragDropEvent, const FSequencerDragDropParams& DragDropParams) override;
+
+private:
+
+	TSharedRef<SWidget> BuildSubMenu(FOnAssetSelected OnAssetSelected, FOnAssetEnterPressed OnAssetEnterPressed);
+	void OnAssetSelected(const FAssetData& AssetData, UMovieSceneTrack* Track);
+	void OnAssetEnterPressed(const TArray<FAssetData>& AssetData, UMovieSceneTrack* Track);
+	void HandleAddMenuEntryExecute();
+};
diff --git a/Engine/Plugins/Experimental/SubtitlesAndClosedCaptions/Source/SubtitlesAndClosedCaptionsEditor/Public/SubtitlesAndClosedCaptionsEditorModule.h b/Engine/Plugins/Experimental/SubtitlesAndClosedCaptions/Source/SubtitlesAndClosedCaptionsEditor/Public/SubtitlesAndClosedCaptionsEditorModule.h
new file mode 100644
--- /dev/null
+++ b/Engine/Plugins/Experimental/SubtitlesAndClosedCaptions/Source/SubtitlesAndClosedCaptionsEditor/Public/SubtitlesAndClosedCaptionsEditorModule.h
@@ -0,0 +1,16 @@
+// Copyright Epic Games, Inc. All Rights Reserved.
+#pragma once
+
+#include "Internationalization/Text.h"
+#include "Modules/ModuleInterface.h"
+
+class ISubtitlesAndClosedCaptionsEditorModule : public IModuleInterface
+{
+public:
+	static FText GetAssetTypeCategory()
+	{
+		return AssetTypeCategory;
+	}
+protected:
+	static FText AssetTypeCategory;
+};
diff --git a/Engine/Plugins/Experimental/SubtitlesAndClosedCaptions/Source/SubtitlesAndClosedCaptionsEditor/SubtitlesAndClosedCaptionsEditor.Build.cs b/Engine/Plugins/Experimental/SubtitlesAndClosedCaptions/Source/SubtitlesAndClosedCaptionsEditor/SubtitlesAndClosedCaptionsEditor.Build.cs
new file mode 100644
--- /dev/null
+++ b/Engine/Plugins/Experimental/SubtitlesAndClosedCaptions/Source/SubtitlesAndClosedCaptionsEditor/SubtitlesAndClosedCaptionsEditor.Build.cs
@@ -0,0 +1,39 @@
+// Copyright Epic Games, Inc. All Rights Reserved.
+
+namespace UnrealBuildTool.Rules
+{
+	public class SubtitlesAndClosedCaptionsEditor : ModuleRules
+	{
+		public SubtitlesAndClosedCaptionsEditor(ReadOnlyTargetRules Target) : base(Target)
+		{
+			PCHUsage = PCHUsageMode.UseExplicitOrSharedPCHs;
+
+			PrivateDependencyModuleNames.AddRange
+			(
+				new string[]
+				{
+					"AssetTools",
+					"LevelSequence",
+					"MovieScene",
+					"Sequencer",
+					"SequencerCore",
+					"Slate",
+					"SlateCore",
+					"SubtitlesAndClosedCaptions",
+				}
+			);
+
+			PublicDependencyModuleNames.AddRange
+			(
+				new string[]
+				{
+					"AssetDefinition",
+					"Core",
+					"CoreUObject",
+					"Engine",
+					"UnrealEd"
+				}
+			);
+		}
+	}
+}
diff --git a/Engine/Plugins/Experimental/SubtitlesAndClosedCaptions/Source/SubtitlesAndClosedCaptionsTest/Private/SubtitlesAndClosedCaptionsMovieSceneTests.cpp b/Engine/Plugins/Experimental/SubtitlesAndClosedCaptions/Source/SubtitlesAndClosedCaptionsTest/Private/SubtitlesAndClosedCaptionsMovieSceneTests.cpp
new file mode 100644
--- /dev/null
+++ b/Engine/Plugins/Experimental/SubtitlesAndClosedCaptions/Source/SubtitlesAndClosedCaptionsTest/Private/SubtitlesAndClosedCaptionsMovieSceneTests.cpp
@@ -0,0 +1,162 @@
+// Copyright Epic Games, Inc. All Rights Reserved.
+
+#include "Misc/AutomationTest.h"
+
+#include "Engine/World.h"
+#include "MovieSceneSubtitlesSystem.h"
+#include "Subtitles/SubtitlesAndClosedCaptionsDelegates.h"
+#include "SubtitlesSubsystem.h"
+#include "UObject/ObjectMacros.h"
+#include "UObject/Package.h"
+#include "UObject/UObjectGlobals.h"
+#include "TimerManager.h"
+
+#include "SubtitlesAndClosedCaptions/Public/MovieSceneSubtitlesSystem.h" // included to test internal FEvaluateSubtitles struct
+
+
+#if WITH_DEV_AUTOMATION_TESTS
+
+using namespace UE::MovieScene;
+
+namespace SubtitlesAndClosedCaptions::Test
+{
+	const FFrameRate TestFrameRate(30, 1);
+	constexpr float SubtitleDurationSeconds = 4.f;
+	constexpr int SubtitleDurationFrames = SubtitleDurationSeconds * 30;
+
+	// Helper struct for initializing tests so the same thing doesn't have to be done in every class
+	struct FMovieSceneSubtitlesTest
+	{
+		FMovieSceneSubtitlesTest()
+			: World(NewObject<UWorld>(GetTransientPackage(), NAME_None, RF_Transient))
+			, Subsystem(NewObject<USubtitlesSubsystem>(World, NAME_None, RF_Transient))
+			, Subtitle(NewObject<USubtitleAssetUserData>(GetTransientPackage(), NAME_None, RF_Transient))
+			, Section(NewObject<UMovieSceneSubtitleSection>(GetTransientPackage(), NAME_None, RF_Transient))
+			, SequenceInstance(TSharedPtr<FSharedPlaybackState>(MakeShareable(new FSharedPlaybackState())).ToSharedRef())
+			, SubtitleData({ Section })
+			, EvaluateTask(nullptr)
+		{
+			Subtitle->Duration = SubtitleDurationSeconds;
+			Section->SetRange(TRange<FFrameNumber>(1, SubtitleDurationFrames));
+			Subsystem->BindDelegates();
+			Section->SetSubtitle(*Subtitle);
+		}
+
+		UWorld* World = nullptr;
+		USubtitlesSubsystem* Subsystem = nullptr;
+		USubtitleAssetUserData* Subtitle = nullptr;
+		UMovieSceneSubtitleSection* Section = nullptr;
+
+		FSequenceInstance SequenceInstance;
+		FSubtitleDataComponent SubtitleData;
+		EMovieScenePlayerStatus::Type LastStatus;
+		FEvaluateSubtitles EvaluateTask;
+	};
+
+	// Verify base case that a subtitle can be queued (and therefore displayed) by the subtitle subsystem
+	// Ensures basic functionality that subtitles can actually be added to the subsystem.
+	// Test content is all transient and uses a new USubtitlesSubsystem that has no UWorld.
+	IMPLEMENT_SIMPLE_AUTOMATION_TEST(UMovieSceneSubtitlesSystem_EvaluateForwardsQueuesSubtitle, "Subtitles.MovieSceneSubtitlesSystem.EvaluateForwardsQueuesSubtitle", EAutomationTestFlags::EditorContext | EAutomationTestFlags::EngineFilter)
+
+		bool UMovieSceneSubtitlesSystem_EvaluateForwardsQueuesSubtitle::RunTest(const FString&)
+	{
+		FMovieSceneSubtitlesTest Test;
+
+		// Evaluate a playing frame.
+		const EMovieScenePlayerStatus::Type CurrentStatus = EMovieScenePlayerStatus::Type::Playing;
+		Test.SequenceInstance.SetContext(FMovieSceneContext(FMovieSceneEvaluationRange(FFrameTime(1), TestFrameRate), CurrentStatus));
+		Test.EvaluateTask.Evaluate(Test.SequenceInstance, Test.SubtitleData, Test.LastStatus);
+
+		const FActiveSubtitle* FoundSubtitle = Test.Subsystem->GetActiveSubtitles().FindByPredicate(
+			[Subtitle = Test.Subtitle](const FActiveSubtitle& ActiveSubtitle) {return Subtitle == ActiveSubtitle.Subtitle; });
+
+		UTEST_NOT_NULL("Subtitle is active", FoundSubtitle);
+
+		return true;
+	}
+
+	// Verify that all subtitles can be stopped by the subtitle movie scene system
+	// Ensures that when playing/stopping/scrubbing in editor that any active subtitles can be reset on movie scene status change.
+	// e.g. When scrubbing through sequencer during evaluation subtitles can be enqueued. Stop all of these and start fresh when playing forwards.
+	// Test content is all transient and uses a new USubtitlesSubsystem that has no UWorld.
+	IMPLEMENT_SIMPLE_AUTOMATION_TEST(UMovieSceneSubtitlesSystem_EvaluateHasStartedStopsAllSubtitles, "Subtitles.MovieSceneSubtitlesSystem.EvaluateHasStartedStopsAllSubtitles", EAutomationTestFlags::EditorContext | EAutomationTestFlags::EngineFilter)
+
+		bool UMovieSceneSubtitlesSystem_EvaluateHasStartedStopsAllSubtitles::RunTest(const FString&)
+	{
+		FMovieSceneSubtitlesTest Test;
+		const USubtitleAssetUserData* Subtitle1 = NewObject<USubtitleAssetUserData>(GetTransientPackage(), NAME_None, RF_Transient);
+		const USubtitleAssetUserData* Subtitle2 = NewObject<USubtitleAssetUserData>(GetTransientPackage(), NAME_None, RF_Transient);
+		const USubtitleAssetUserData* Subtitle3 = NewObject<USubtitleAssetUserData>(GetTransientPackage(), NAME_None, RF_Transient);
+
+		Test.Subsystem->QueueSubtitle(FQueueSubtitleParameters{ *Subtitle1 });
+		Test.Subsystem->QueueSubtitle(FQueueSubtitleParameters{ *Subtitle2 });
+		Test.Subsystem->QueueSubtitle(FQueueSubtitleParameters{ *Subtitle3 });
+		UTEST_EQUAL("All subtitles active", Test.Subsystem->GetActiveSubtitles().Num(), 3);
+
+		Test.LastStatus = EMovieScenePlayerStatus::Type::Stopped;
+		EMovieScenePlayerStatus::Type CurrentStatus = EMovieScenePlayerStatus::Type::Playing;
+		Test.SequenceInstance.SetContext(FMovieSceneContext(FMovieSceneEvaluationRange(FFrameTime(0), TestFrameRate), CurrentStatus));
+		Test.EvaluateTask.Evaluate(Test.SequenceInstance, Test.SubtitleData, Test.LastStatus);
+
+		const FActiveSubtitle* FoundSubtitle = Test.Subsystem->GetActiveSubtitles().FindByPredicate(
+			[Subtitle = Test.Subtitle](const FActiveSubtitle& ActiveSubtitle) {return Subtitle == ActiveSubtitle.Subtitle; });
+
+		UTEST_NOT_NULL("All subtitles removed before queueing new subtitle", FoundSubtitle);
+
+		return true;
+	}
+
+	// Verify that all subtitles can be displayed indefinitely
+	// Ensures that when you pause on an active subtitle that its duration does not expire
+	// Test content is all transient and uses a new USubtitlesSubsystem that has no UWorld.
+	IMPLEMENT_SIMPLE_AUTOMATION_TEST(UMovieSceneSubtitlesSystem_EvaluateStoppedQueuesInfiniteDuration, "Subtitles.MovieSceneSubtitlesSystem.EvaluateStoppedQueuesInfiniteDuration", EAutomationTestFlags::EditorContext | EAutomationTestFlags::EngineFilter)
+
+		bool UMovieSceneSubtitlesSystem_EvaluateStoppedQueuesInfiniteDuration::RunTest(const FString&)
+	{
+		FMovieSceneSubtitlesTest Test;
+		Test.LastStatus = EMovieScenePlayerStatus::Type::Playing;
+		EMovieScenePlayerStatus::Type CurrentStatus = EMovieScenePlayerStatus::Type::Stopped;
+		Test.SequenceInstance.SetContext(FMovieSceneContext(FMovieSceneEvaluationRange(FFrameTime(30), TestFrameRate), CurrentStatus));
+
+		Test.EvaluateTask.Evaluate(Test.SequenceInstance, Test.SubtitleData, Test.LastStatus);
+
+		UTEST_TRUE("A Subtitle is active", !Test.Subsystem->GetActiveSubtitles().IsEmpty());
+
+		const float ExpectedDuration = FLT_MAX;
+		const float ActualDuration = Test.World->GetTimerManager().GetTimerRemaining(Test.Subsystem->GetActiveSubtitles()[0].DurationTimerHandle);
+		UTEST_EQUAL_TOLERANCE("Subtitle played mid-way is queued with remaining duration", ActualDuration, ExpectedDuration, FLT_EPSILON);
+
+		return true;
+	}
+
+	// Verify that the subtitle section only queues its subtitle once
+	// Ensures that the subtitle section doesn't spam the subtitle subsystem
+	// Test content is all transient and uses a new USubtitlesSubsystem that has no UWorld.
+	IMPLEMENT_SIMPLE_AUTOMATION_TEST(UMovieSceneSubtitlesSystem_EvaluateDontQueueIfSubtitleAlreadyActive, "Subtitles.MovieSceneSubtitlesSystem.EvaluateDontQueueIfSubtitleAlreadyActive", EAutomationTestFlags::EditorContext | EAutomationTestFlags::EngineFilter)
+
+		bool UMovieSceneSubtitlesSystem_EvaluateDontQueueIfSubtitleAlreadyActive::RunTest(const FString&)
+	{
+		FMovieSceneSubtitlesTest Test;
+		Test.Subtitle->Duration = 1.f;
+		EMovieScenePlayerStatus::Type CurrentStatus = EMovieScenePlayerStatus::Type::Playing;
+
+		Test.SequenceInstance.SetContext(FMovieSceneContext(FMovieSceneEvaluationRange(FFrameTime(0), TestFrameRate), CurrentStatus));
+		Test.EvaluateTask.Evaluate(Test.SequenceInstance, Test.SubtitleData, Test.LastStatus);
+
+		Test.SequenceInstance.SetContext(FMovieSceneContext(FMovieSceneEvaluationRange(FFrameTime(1), TestFrameRate), CurrentStatus));
+		Test.EvaluateTask.Evaluate(Test.SequenceInstance, Test.SubtitleData, Test.LastStatus);
+
+		Test.SequenceInstance.SetContext(FMovieSceneContext(FMovieSceneEvaluationRange(FFrameTime(2), TestFrameRate), CurrentStatus));
+		Test.EvaluateTask.Evaluate(Test.SequenceInstance, Test.SubtitleData, Test.LastStatus);
+
+		Test.SequenceInstance.SetContext(FMovieSceneContext(FMovieSceneEvaluationRange(FFrameTime(3), TestFrameRate), CurrentStatus));
+		Test.EvaluateTask.Evaluate(Test.SequenceInstance, Test.SubtitleData, Test.LastStatus);
+
+		UTEST_EQUAL("SubtitleSection's Subtitle is only queued once", Test.Subsystem->GetActiveSubtitles().Num(), 1);
+
+		return true;
+	}
+
+} // namespace SubtitlesAndClosedCaptions::Test
+
+#endif
diff --git a/Engine/Plugins/Experimental/SubtitlesAndClosedCaptions/Source/SubtitlesAndClosedCaptionsTest/Private/SubtitlesAndClosedCaptionsTestModule.cpp b/Engine/Plugins/Experimental/SubtitlesAndClosedCaptions/Source/SubtitlesAndClosedCaptionsTest/Private/SubtitlesAndClosedCaptionsTestModule.cpp
new file mode 100644
--- /dev/null
+++ b/Engine/Plugins/Experimental/SubtitlesAndClosedCaptions/Source/SubtitlesAndClosedCaptionsTest/Private/SubtitlesAndClosedCaptionsTestModule.cpp
@@ -0,0 +1,15 @@
+// Copyright Epic Games, Inc. All Rights Reserved.
+
+#include "Modules/ModuleInterface.h"
+#include "Modules/ModuleManager.h"
+
+class SUBTITLESANDCLOSEDCAPTIONSTEST_API FSubtitlesAndClosedCaptionsTestModule : public IModuleInterface
+{
+public:
+	virtual void StartupModule() override
+	{
+		FModuleManager::Get().LoadModuleChecked("SubtitlesAndClosedCaptions");
+	}
+};
+
+IMPLEMENT_MODULE(FSubtitlesAndClosedCaptionsTestModule, SubtitlesAndClosedCaptionsTest);
diff --git a/Engine/Plugins/Experimental/SubtitlesAndClosedCaptions/Source/SubtitlesAndClosedCaptionsTest/Private/SubtitlesAndClosedCaptionsTests.cpp b/Engine/Plugins/Experimental/SubtitlesAndClosedCaptions/Source/SubtitlesAndClosedCaptionsTest/Private/SubtitlesAndClosedCaptionsTests.cpp
new file mode 100644
--- /dev/null
+++ b/Engine/Plugins/Experimental/SubtitlesAndClosedCaptions/Source/SubtitlesAndClosedCaptionsTest/Private/SubtitlesAndClosedCaptionsTests.cpp
@@ -0,0 +1,222 @@
+// Copyright Epic Games, Inc. All Rights Reserved.
+
+#include "Misc/AutomationTest.h"
+
+#include "Engine/World.h"
+#include "Subtitles/SubtitlesAndClosedCaptionsDelegates.h"
+#include "SubtitlesSubsystem.h"
+#include "TimerManager.h"
+#include "UObject/ObjectMacros.h"
+#include "UObject/Package.h"
+#include "UObject/UObjectGlobals.h"
+
+#if WITH_DEV_AUTOMATION_TESTS
+
+namespace SubtitlesAndClosedCaptions::Test
+{
+	// Helper struct for initializing tests so the same thing doesn't have to be done in every class
+	struct FSubtitlesTest
+	{
+		FSubtitlesTest()
+			: World(NewObject<UWorld>(GetTransientPackage(), NAME_None, RF_Transient))
+			, Subsystem(NewObject<USubtitlesSubsystem>(World, NAME_None, RF_Transient))
+		{
+			Subsystem->BindDelegates();
+		}
+
+		UWorld* World = nullptr;
+		USubtitlesSubsystem* Subsystem = nullptr;
+	};
+
+	// Verify that a subtitle can be added to the subtitle subsystem.
+	// Ensures basic functionality that subtitles can actually be added to the subsystem.
+	// Test content is all transient and uses a new USubtitlesSubsystem that has no UWorld.
+	IMPLEMENT_SIMPLE_AUTOMATION_TEST(USubtitlesSubsystem_QueueSubtitle, "Subtitles.SubtitlesSubsystem.QueueSubtitle", EAutomationTestFlags::EditorContext | EAutomationTestFlags::EngineFilter)
+
+		bool USubtitlesSubsystem_QueueSubtitle::RunTest(const FString&)
+	{
+		FSubtitlesTest Test;
+		const USubtitleAssetUserData* Subtitle = NewObject<USubtitleAssetUserData>(GetTransientPackage(), NAME_None, RF_Transient);
+
+		Test.Subsystem->QueueSubtitle(FQueueSubtitleParameters{ *Subtitle });
+		UTEST_EQUAL("Active subtitle successfully added", Test.Subsystem->GetActiveSubtitles().Num(), 1);
+
+		return true;
+	}
+
+	// Verify that a subtitle can be stopped manually before elapsing its entire duration
+	// Ensures that gameplay or any systems have the ability to cut off a subtitle
+	// Test content is all transient and uses a new USubtitlesSubsystem that has no UWorld.
+	IMPLEMENT_SIMPLE_AUTOMATION_TEST(USubtitlesSubsystem_StopSubtitle, "Subtitles.SubtitlesSubsystem.StopSubtitle", EAutomationTestFlags::EditorContext | EAutomationTestFlags::EngineFilter)
+
+		bool USubtitlesSubsystem_StopSubtitle::RunTest(const FString&)
+	{
+		FSubtitlesTest Test;
+		const USubtitleAssetUserData* Subtitle1 = NewObject<USubtitleAssetUserData>(GetTransientPackage(), NAME_None, RF_Transient);
+		const USubtitleAssetUserData* Subtitle2 = NewObject<USubtitleAssetUserData>(GetTransientPackage(), NAME_None, RF_Transient);
+
+		Test.Subsystem->QueueSubtitle(FQueueSubtitleParameters{ *Subtitle1 });
+		Test.Subsystem->QueueSubtitle(FQueueSubtitleParameters{ *Subtitle2 });
+		UTEST_TRUE("A subtitle is active", !Test.Subsystem->GetActiveSubtitles().IsEmpty());
+
+		Test.Subsystem->StopSubtitle(*Subtitle1);
+		const FActiveSubtitle* FoundSubtitle1 = Test.Subsystem->GetActiveSubtitles().FindByPredicate([Subtitle1](const FActiveSubtitle& ActiveSubtitle) { return ActiveSubtitle.Subtitle == Subtitle1; });
+		UTEST_NULL("Subtitle1 was stopped", FoundSubtitle1);
+
+		const FActiveSubtitle* FoundSubtitle2 = Test.Subsystem->GetActiveSubtitles().FindByPredicate([Subtitle2](const FActiveSubtitle& ActiveSubtitle) { return ActiveSubtitle.Subtitle == Subtitle2; });
+		UTEST_NOT_NULL("Subtitle2 is still active", FoundSubtitle2);
+
+		return true;
+	}
+
+	// Verify that subtitles are prioritized by descending order of priority
+	// Ensures that when there are many subtitles trying to play at once (and being added/removed) that the prioritization sort order is maintained.
+	// Test content is all transient and uses a new USubtitlesSubsystem that has no UWorld.
+	IMPLEMENT_SIMPLE_AUTOMATION_TEST(USubtitlesSubsystem_PrioritizeSubtitles, "Subtitles.SubtitlesSubsystem.PrioritizeSubtitles", EAutomationTestFlags::EditorContext | EAutomationTestFlags::EngineFilter)
+
+		bool USubtitlesSubsystem_PrioritizeSubtitles::RunTest(const FString&)
+	{
+		FSubtitlesTest Test;
+
+		// prioritized subtitles with priorities from 0 to 3, where 0 is the lowest priority and 3 is the highest priority
+		// 0, 1, 2, 3
+		constexpr uint32 NumSubtitles = 4;
+		TStaticArray< const USubtitleAssetUserData*, NumSubtitles> Subtitles;
+		for (int i = 0; i < NumSubtitles; ++i)
+		{
+			USubtitleAssetUserData* Subtitle = NewObject<USubtitleAssetUserData>(GetTransientPackage(), NAME_None, RF_Transient);
+			Subtitle->Priority = static_cast<float>(i);
+			Subtitles[i] = Subtitle;
+			Test.Subsystem->QueueSubtitle(FQueueSubtitleParameters{ *Subtitle });
+		}
+
+		UTEST_TRUE("Add subtitles added", Test.Subsystem->GetActiveSubtitles().Num() == NumSubtitles);
+		UTEST_TRUE("Highest priority subtitle (3) is first index", Test.Subsystem->GetActiveSubtitles()[0].Subtitle == Subtitles[3]);
+
+		Test.Subsystem->StopSubtitle(*Subtitles[3]);
+		UTEST_TRUE("After removing highest priority subtitle the new Highest priority subtitle is 2", Test.Subsystem->GetActiveSubtitles()[0].Subtitle == Subtitles[2]);
+
+		Test.Subsystem->StopSubtitle(*Subtitles[0]);
+		UTEST_TRUE("After removing lower priority subtitle the highest priority is still the same", Test.Subsystem->GetActiveSubtitles()[0].Subtitle == Subtitles[2]);
+
+		return true;
+	}
+
+	// Verify that subtitle API supports displaying a subtitle for a specified duration, regardless of the duration on the asset.
+	// Ensures that a subtitle can start half way through its duration, or play infinitely and be manually stopped 
+	// Test content is all transient and uses a new USubtitlesSubsystem that has no UWorld.
+	IMPLEMENT_SIMPLE_AUTOMATION_TEST(USubtitlesSubsystem_OverrideDuration, "Subtitles.SubtitlesSubsystem.OverrideDuration", EAutomationTestFlags::EditorContext | EAutomationTestFlags::EngineFilter)
+
+		bool USubtitlesSubsystem_OverrideDuration::RunTest(const FString&)
+	{
+		FSubtitlesTest Test;
+		USubtitleAssetUserData* Subtitle = NewObject<USubtitleAssetUserData>(GetTransientPackage(), NAME_None, RF_Transient);
+		Subtitle->Duration = 0.f; // invalid duration as < SubtitleMinDuration
+
+		constexpr float ExpectedDuration = 42.f;
+		Test.Subsystem->QueueSubtitle(FQueueSubtitleParameters{ *Subtitle, ExpectedDuration });
+		UTEST_TRUE("A subtitle is active", !Test.Subsystem->GetActiveSubtitles().IsEmpty());
+
+		const float ActualDuration = Test.World->GetTimerManager().GetTimerRemaining(Test.Subsystem->GetActiveSubtitles()[0].DurationTimerHandle);
+		UTEST_EQUAL_TOLERANCE("The expected subtitle has the expected duration", ActualDuration, ExpectedDuration, FLT_EPSILON);
+
+		return true;
+	}
+
+	// Verify that subtitle API uses SubtitleMinDuration when a duration isn't set or otherwise is zero, regardless of the duration on the asset.
+	// Test content is all transient and uses a new USubtitlesSubsystem that has no UWorld.
+	IMPLEMENT_SIMPLE_AUTOMATION_TEST(USubtitlesSubsystem_ClampDuration, "Subtitles.SubtitlesSubsystem.ClampDuration", EAutomationTestFlags::EditorContext | EAutomationTestFlags::EngineFilter)
+
+	bool USubtitlesSubsystem_ClampDuration::RunTest(const FString&)
+	{
+		FSubtitlesTest Test;
+		USubtitleAssetUserData* Subtitle = NewObject<USubtitleAssetUserData>(GetTransientPackage(), NAME_None, RF_Transient);
+		Subtitle->Duration = 0.f; // invalid duration as < SubtitleMinDuration
+
+		constexpr float ExpectedDuration = SubtitleMinDuration;
+		Test.Subsystem->QueueSubtitle(FQueueSubtitleParameters{ *Subtitle });
+		UTEST_TRUE("A subtitle is active", !Test.Subsystem->GetActiveSubtitles().IsEmpty());
+
+		const float ActualDuration = Test.World->GetTimerManager().GetTimerRemaining(Test.Subsystem->GetActiveSubtitles()[0].DurationTimerHandle);
+		UTEST_EQUAL_TOLERANCE("The expected subtitle has the expected duration", ActualDuration, ExpectedDuration, FLT_EPSILON);
+
+		return true;
+	}
+
+	// Verify that a user can check if the subtitle is already active
+	// Ensures that systems like movie scene can check if a subtitle is already active before trying to queue one up every frame
+	// Test content is all transient and uses a new USubtitlesSubsystem that has no UWorld.
+	IMPLEMENT_SIMPLE_AUTOMATION_TEST(USubtitlesSubsystem_IsSubtitleActive, "Subtitles.SubtitlesSubsystem.IsSubtitleActive", EAutomationTestFlags::EditorContext | EAutomationTestFlags::EngineFilter)
+
+		bool USubtitlesSubsystem_IsSubtitleActive::RunTest(const FString&)
+	{
+		FSubtitlesTest Test;
+		const USubtitleAssetUserData* Subtitle = NewObject<USubtitleAssetUserData>(GetTransientPackage(), NAME_None, RF_Transient);
+
+		UTEST_TRUE("Subtitle is not active", !Test.Subsystem->IsSubtitleActive(*Subtitle));
+
+		Test.Subsystem->QueueSubtitle(FQueueSubtitleParameters{ *Subtitle });
+		UTEST_TRUE("A subtitle is active", Test.Subsystem->IsSubtitleActive(*Subtitle));
+		UTEST_SAME_PTR("The expected subtitle is active", Subtitle, Test.Subsystem->GetActiveSubtitles()[0].Subtitle.Get());
+
+		return true;
+	}
+
+	// Verify that all subtitles can be stopped at once
+	// Ensures that gameplay or any systems have the ability to stop all subtitles if for whatever reason there is a context or state that needs to reset subtitles
+	// Test content is all transient and uses a new USubtitlesSubsystem that has no UWorld.
+	IMPLEMENT_SIMPLE_AUTOMATION_TEST(USubtitlesSubsystem_StopAllSubtitles, "Subtitles.SubtitlesSubsystem.StopAllSubtitles", EAutomationTestFlags::EditorContext | EAutomationTestFlags::EngineFilter)
+
+		bool USubtitlesSubsystem_StopAllSubtitles::RunTest(const FString&)
+	{
+		FSubtitlesTest Test;
+		const USubtitleAssetUserData* Subtitle1 = NewObject<USubtitleAssetUserData>(GetTransientPackage(), NAME_None, RF_Transient);
+		
+		USubtitleAssetUserData* Subtitle2 = NewObject<USubtitleAssetUserData>(GetTransientPackage(), NAME_None, RF_Transient);
+		Subtitle2->StartOffset = 1.f;
+
+		Test.Subsystem->QueueSubtitle(FQueueSubtitleParameters{ *Subtitle1 });
+		Test.Subsystem->QueueSubtitle(FQueueSubtitleParameters{ *Subtitle2 });
+		UTEST_TRUE("A subtitle is active", !Test.Subsystem->GetActiveSubtitles().IsEmpty());
+
+		Test.Subsystem->StopAllSubtitles();
+		const FActiveSubtitle* FoundSubtitle1 = Test.Subsystem->GetActiveSubtitles().FindByPredicate([Subtitle1](const FActiveSubtitle& ActiveSubtitle) { return ActiveSubtitle.Subtitle == Subtitle1; });
+		UTEST_NULL("Subtitle1 was stopped", FoundSubtitle1);
+
+		const FActiveSubtitle* FoundSubtitle2 = Test.Subsystem->GetActiveSubtitles().FindByPredicate([Subtitle2](const FActiveSubtitle& ActiveSubtitle) { return ActiveSubtitle.Subtitle == Subtitle2; });
+		UTEST_NULL("Subtitle2 was stopped", FoundSubtitle2);
+
+		return true;
+	}
+
+	// Verify that delayed subtitles do not immediately join the queue, and that they do join the queue after their delay timer expires.
+	// Ensures that migrated subtitles using a delayed start will continue to start at their expected times.
+	// Test content is all transient and uses a new USubtitlesSubsystem that has no UWorld.
+	IMPLEMENT_SIMPLE_AUTOMATION_TEST(USubtitlesSubsystem_DelayedOffset, "Subtitles.SubtitlesSubsystem.CheckDelayedOffsets", EAutomationTestFlags::EditorContext | EAutomationTestFlags::EngineFilter)
+		bool USubtitlesSubsystem_DelayedOffset::RunTest(const FString&)
+	{
+		FSubtitlesTest Test;
+		USubtitleAssetUserData* InstantSubtitle = NewObject<USubtitleAssetUserData>(GetTransientPackage(), NAME_None, RF_Transient);
+		USubtitleAssetUserData* DelayedSubtitle = NewObject<USubtitleAssetUserData>(GetTransientPackage(), NAME_None, RF_Transient);
+
+		// Both subtitles have the default duration of 3 seconds.
+		DelayedSubtitle->StartOffset = 1.f;
+
+		// The delayed subtitle should have priority when its delay timer expires.
+		DelayedSubtitle->Priority = 999;
+		InstantSubtitle->Priority = 1;
+		
+		Test.Subsystem->QueueSubtitle(FQueueSubtitleParameters{ *InstantSubtitle });
+		Test.Subsystem->QueueSubtitle(FQueueSubtitleParameters{ *DelayedSubtitle });
+		UTEST_TRUE("Only the instant subtitle is in the queue.", Test.Subsystem->GetActiveSubtitles().Num() == 1);
+
+		Test.Subsystem->TestActivatingDelayedSubtitle(*DelayedSubtitle);
+		UTEST_TRUE("There are two subtitles in the queue now.", Test.Subsystem->GetActiveSubtitles().Num() == 2);
+		UTEST_SAME_PTR("The delayed subtitle is now active.", DelayedSubtitle, Test.Subsystem->GetTopRankedSubtitle());
+
+		return true;
+	}
+
+} // namespace SubtitlesAndClosedCaptions::Test
+
+#endif
diff --git a/Engine/Plugins/Experimental/SubtitlesAndClosedCaptions/Source/SubtitlesAndClosedCaptionsTest/SubtitlesAndClosedCaptionsTest.Build.cs b/Engine/Plugins/Experimental/SubtitlesAndClosedCaptions/Source/SubtitlesAndClosedCaptionsTest/SubtitlesAndClosedCaptionsTest.Build.cs
new file mode 100644
--- /dev/null
+++ b/Engine/Plugins/Experimental/SubtitlesAndClosedCaptions/Source/SubtitlesAndClosedCaptionsTest/SubtitlesAndClosedCaptionsTest.Build.cs
@@ -0,0 +1,34 @@
+// Copyright Epic Games, Inc. All Rights Reserved.
+
+using UnrealBuildTool;
+using System.IO;
+
+namespace UnrealBuildTool.Rules
+{
+	public class SubtitlesAndClosedCaptionsTest : ModuleRules
+	{
+		public SubtitlesAndClosedCaptionsTest(ReadOnlyTargetRules Target) : base(Target)
+		{
+			PCHUsage = PCHUsageMode.UseExplicitOrSharedPCHs;
+
+			PrivateDependencyModuleNames.AddRange
+			(
+				new string[]
+				{
+					"Core",
+					"CoreUObject",
+					"Engine",
+					"MovieScene"
+				}
+			);
+
+			PublicDependencyModuleNames.AddRange
+			(
+				new string[]
+				{
+					"SubtitlesAndClosedCaptions",
+				}
+			);
+		}
+	}
+}
diff --git a/Engine/Plugins/Experimental/SubtitlesAndClosedCaptions/SubtitlesAndClosedCaptions.uplugin b/Engine/Plugins/Experimental/SubtitlesAndClosedCaptions/SubtitlesAndClosedCaptions.uplugin
new file mode 100644
--- /dev/null
+++ b/Engine/Plugins/Experimental/SubtitlesAndClosedCaptions/SubtitlesAndClosedCaptions.uplugin
@@ -0,0 +1,41 @@
+{
+	"FileVersion": 3,
+	"Version": 1000,
+	"VersionName": "1.0.0.0",
+	"FriendlyName": "Subtitles and Closed Captions",
+	"Description": "Standalone plugin for displaying Subtitles and Closed Captions",
+	"Category": "Audio",
+	"CreatedBy": "Epic Games, Inc.",
+	"CreatedByURL": "https://epicgames.com/",
+	"DocsURL": "",
+	"MarketplaceURL": "",
+	"SupportURL": "",
+	"EnabledByDefault": false,
+	"CanContainContent": true,
+	"IsBetaVersion": false,
+	"IsExperimentalVersion": true,
+	"Installed": false,
+	"Modules": [
+		{
+			"Name": "SubtitlesAndClosedCaptions",
+			"Type": "Runtime",
+			"LoadingPhase": "Default",
+			"PlatformAllowList": [
+				"Win64"
+			]
+		},
+		{
+			"Name": "SubtitlesAndClosedCaptionsTest",
+			"Type": "Runtime",
+			"LoadingPhase": "Default",
+			"PlatformAllowList": [
+				"Win64"
+			]
+		},
+		{
+			"Name": "SubtitlesAndClosedCaptionsEditor",
+			"Type": "UncookedOnly",
+			"LoadingPhase": "Default"
+		}
+	]
+}
