Fix for issue 891649, Static Analysis warning.

diff --git a/Engine/Plugins/Experimental/SubtitlesAndClosedCaptions/Source/SubtitlesAndClosedCaptionsTest/Private/SubtitlesAndClosedCaptionsMovieSceneTests.cpp b/Engine/Plugins/Experimental/SubtitlesAndClosedCaptions/Source/SubtitlesAndClosedCaptionsTest/Private/SubtitlesAndClosedCaptionsMovieSceneTests.cpp
--- a/Engine/Plugins/Experimental/SubtitlesAndClosedCaptions/Source/SubtitlesAndClosedCaptionsTest/Private/SubtitlesAndClosedCaptionsMovieSceneTests.cpp
+++ b/Engine/Plugins/Experimental/SubtitlesAndClosedCaptions/Source/SubtitlesAndClosedCaptionsTest/Private/SubtitlesAndClosedCaptionsMovieSceneTests.cpp
@@ -35,7 +35,7 @@ namespace SubtitlesAndClosedCaptions::Test
 			, PlaybackState(MakeShared<FSharedPlaybackState>())
 			, SequenceInstance(PlaybackState)
 			, SubtitleData({ Section })
-			, EvaluateTask(nullptr)
+			, EvaluateTask(SequenceInstance.GetSharedPlaybackState()->GetLinker())
 		{
 			Subtitle->Duration = SubtitleDurationSeconds;
 			Section->SetRange(TRange<FFrameNumber>(1, SubtitleDurationFrames));
