Fix compile errors with Module.SubtitleSettings.cpp, SubtitleSettings.cpp, and others.

diff --git a/Engine/Plugins/Experimental/SubtitlesAndClosedCaptions/Source/SubtitlesAndClosedCaptions/Private/SubtitleTextBlock.cpp b/Engine/Plugins/Experimental/SubtitlesAndClosedCaptions/Source/SubtitlesAndClosedCaptions/Private/SubtitleTextBlock.cpp
--- a/Engine/Plugins/Experimental/SubtitlesAndClosedCaptions/Source/SubtitlesAndClosedCaptions/Private/SubtitleTextBlock.cpp
+++ b/Engine/Plugins/Experimental/SubtitlesAndClosedCaptions/Source/SubtitlesAndClosedCaptions/Private/SubtitleTextBlock.cpp
@@ -2,6 +2,7 @@
 
 #include "SubtitleTextBlock.h"
 #include "SubtitlesAndClosedCaptionsModule.h"
+#include "Components/TextBlock.h"
 
 void USubtitleTextBlock::NativeConstruct()
 {
diff --git a/Engine/Plugins/Experimental/SubtitlesAndClosedCaptions/Source/SubtitlesAndClosedCaptions/Public/SubtitleTextBlock.h b/Engine/Plugins/Experimental/SubtitlesAndClosedCaptions/Source/SubtitlesAndClosedCaptions/Public/SubtitleTextBlock.h
--- a/Engine/Plugins/Experimental/SubtitlesAndClosedCaptions/Source/SubtitlesAndClosedCaptions/Public/SubtitleTextBlock.h
+++ b/Engine/Plugins/Experimental/SubtitlesAndClosedCaptions/Source/SubtitlesAndClosedCaptions/Public/SubtitleTextBlock.h
@@ -2,12 +2,12 @@
 
 #pragma once
 #include "Blueprint/UserWidget.h"
-#include "Components/TextBlock.h"
-
 #include "Subtitles/SubtitlesAndClosedCaptionsDelegates.h"
 
 #include "SubtitleTextBlock.generated.h"
 
+class UTextBlock;
+
 UCLASS(Abstract, BlueprintType, Blueprintable, meta=( DontUseGenericSpawnObject="True", DisableNativeTick), MinimalAPI)
 class USubtitleTextBlock : public UUserWidget
 {
diff --git a/Engine/Source/Runtime/MovieScene/Public/Evaluation/PreAnimatedState/MovieScenePreAnimatedStateTypes.h b/Engine/Source/Runtime/MovieScene/Public/Evaluation/PreAnimatedState/MovieScenePreAnimatedStateTypes.h
--- a/Engine/Source/Runtime/MovieScene/Public/Evaluation/PreAnimatedState/MovieScenePreAnimatedStateTypes.h
+++ b/Engine/Source/Runtime/MovieScene/Public/Evaluation/PreAnimatedState/MovieScenePreAnimatedStateTypes.h
@@ -28,7 +28,7 @@ struct FCachePreAnimatedValueParams
 	bool bForcePersist = false;
 };
 
-enum EPreAnimatedStorageRequirement : uint8
+enum class EPreAnimatedStorageRequirement : uint8
 {
 	None,
 	Persistent,
