Fix Subtitle tests regression by making tests also initialize the UMG widget.

diff --git a/Engine/Plugins/Experimental/SubtitlesAndClosedCaptions/Source/SubtitlesAndClosedCaptions/Private/SubtitlesSubsystem.cpp b/Engine/Plugins/Experimental/SubtitlesAndClosedCaptions/Source/SubtitlesAndClosedCaptions/Private/SubtitlesSubsystem.cpp
--- a/Engine/Plugins/Experimental/SubtitlesAndClosedCaptions/Source/SubtitlesAndClosedCaptions/Private/SubtitlesSubsystem.cpp
+++ b/Engine/Plugins/Experimental/SubtitlesAndClosedCaptions/Source/SubtitlesAndClosedCaptions/Private/SubtitlesSubsystem.cpp
@@ -23,18 +23,18 @@
 
 void USubtitlesSubsystem::Initialize(FSubsystemCollectionBase& Collection)
 {
-	if (TryCreateUMGWidget())
-	{
-		BindDelegates();
-	}
+	BindDelegates();
 }
 
 void USubtitlesSubsystem::BindDelegates()
 {
-	FSubtitlesAndClosedCaptionsDelegates::QueueSubtitle.BindUObject(this, &USubtitlesSubsystem::QueueSubtitle);
-	FSubtitlesAndClosedCaptionsDelegates::IsSubtitleActive.BindUObject(this, &USubtitlesSubsystem::IsSubtitleActive);
-	FSubtitlesAndClosedCaptionsDelegates::StopSubtitle.BindUObject(this, &USubtitlesSubsystem::StopSubtitle);
-	FSubtitlesAndClosedCaptionsDelegates::StopAllSubtitles.BindUObject(this, &USubtitlesSubsystem::StopAllSubtitles);
+	if (TryCreateUMGWidget())
+	{
+		FSubtitlesAndClosedCaptionsDelegates::QueueSubtitle.BindUObject(this, &USubtitlesSubsystem::QueueSubtitle);
+		FSubtitlesAndClosedCaptionsDelegates::IsSubtitleActive.BindUObject(this, &USubtitlesSubsystem::IsSubtitleActive);
+		FSubtitlesAndClosedCaptionsDelegates::StopSubtitle.BindUObject(this, &USubtitlesSubsystem::StopSubtitle);
+		FSubtitlesAndClosedCaptionsDelegates::StopAllSubtitles.BindUObject(this, &USubtitlesSubsystem::StopAllSubtitles);
+	}
 }
 
 
