[SubtitlesAndClosedCaptions]
* Fixed non unity compile error

diff --git a/Engine/Plugins/Experimental/SubtitlesAndClosedCaptions/Source/SubtitlesAndClosedCaptions/Private/SubtitlesSubsystem.cpp b/Engine/Plugins/Experimental/SubtitlesAndClosedCaptions/Source/SubtitlesAndClosedCaptions/Private/SubtitlesSubsystem.cpp
--- a/Engine/Plugins/Experimental/SubtitlesAndClosedCaptions/Source/SubtitlesAndClosedCaptions/Private/SubtitlesSubsystem.cpp
+++ b/Engine/Plugins/Experimental/SubtitlesAndClosedCaptions/Source/SubtitlesAndClosedCaptions/Private/SubtitlesSubsystem.cpp
@@ -2,6 +2,7 @@
 
 #include "SubtitlesSubsystem.h"
 
+#include "Algo/Find.h"
 #include "Algo/RemoveIf.h"
 #include "CanvasItem.h"
 #include "CanvasTypes.h"
