Fix Compile errors in SubtitleFunctionalTest.cpp, issue 957330

diff --git a/Engine/Plugins/Experimental/SubtitlesAndClosedCaptions/Source/SubtitlesAndClosedCaptions/Public/SubtitlesSubsystem.h b/Engine/Plugins/Experimental/SubtitlesAndClosedCaptions/Source/SubtitlesAndClosedCaptions/Public/SubtitlesSubsystem.h
--- a/Engine/Plugins/Experimental/SubtitlesAndClosedCaptions/Source/SubtitlesAndClosedCaptions/Public/SubtitlesSubsystem.h
+++ b/Engine/Plugins/Experimental/SubtitlesAndClosedCaptions/Source/SubtitlesAndClosedCaptions/Public/SubtitlesSubsystem.h
@@ -85,13 +85,13 @@ protected:
 #endif
 	friend struct SubtitlesAndClosedCaptions::Test::FSubtitlesTest;
 public:
-	const TArray<FActiveSubtitle>& GetActiveSubtitles() const
+	UE_API const TArray<FActiveSubtitle>& GetActiveSubtitles() const
 	{
 		return ActiveSubtitles;
 	}
 
 	// As this is for testing, assume that the number of subtitles has already been checked.
-	const USubtitleAssetUserData* GetTopRankedSubtitle() const
+	UE_API const USubtitleAssetUserData* GetTopRankedSubtitle() const
 	{
 		check(ActiveSubtitles.Num() > 0);
 		const FActiveSubtitle& ActiveSubtitle = ActiveSubtitles[0];
@@ -99,7 +99,7 @@ public:
 		return ActiveSubtitle.Subtitle.Get();
 	}
 
-	void TestActivatingDelayedSubtitle(const UAssetUserData& Data)
+	UE_API void TestActivatingDelayedSubtitle(const UAssetUserData& Data)
 	{
 		const USubtitleAssetUserData* Subtitle = CastChecked<const USubtitleAssetUserData>(&Data);
 		MakeDelayedSubtitleActive(Subtitle);
