Restore functional tests, remove previous Subtitle Test module inside of Subtitles plugin. Portion of review that Convert SubtitlesSubsystem to an AudioEngineSubsystem was removed.

diff --git a/Engine/Plugins/Experimental/SubtitlesAndClosedCaptions/Source/SubtitlesAndClosedCaptionsTest/Private/SubtitlesAndClosedCaptionsMovieSceneTests.cpp b/Engine/Plugins/Experimental/SubtitlesAndClosedCaptions/Source/SubtitlesAndClosedCaptionsTest/Private/SubtitlesAndClosedCaptionsMovieSceneTests.cpp
deleted file mode 100644
diff --git a/Engine/Plugins/Experimental/SubtitlesAndClosedCaptions/Source/SubtitlesAndClosedCaptionsTest/Private/SubtitlesAndClosedCaptionsTestModule.cpp b/Engine/Plugins/Experimental/SubtitlesAndClosedCaptions/Source/SubtitlesAndClosedCaptionsTest/Private/SubtitlesAndClosedCaptionsTestModule.cpp
deleted file mode 100644
diff --git a/Engine/Plugins/Experimental/SubtitlesAndClosedCaptions/Source/SubtitlesAndClosedCaptionsTest/Private/SubtitlesAndClosedCaptionsTests.cpp b/Engine/Plugins/Experimental/SubtitlesAndClosedCaptions/Source/SubtitlesAndClosedCaptionsTest/Private/SubtitlesAndClosedCaptionsTests.cpp
deleted file mode 100644
diff --git a/Engine/Plugins/Experimental/SubtitlesAndClosedCaptions/Source/SubtitlesAndClosedCaptionsTest/SubtitlesAndClosedCaptionsTest.Build.cs b/Engine/Plugins/Experimental/SubtitlesAndClosedCaptions/Source/SubtitlesAndClosedCaptionsTest/SubtitlesAndClosedCaptionsTest.Build.cs
deleted file mode 100644
diff --git a/Engine/Plugins/Experimental/SubtitlesAndClosedCaptions/SubtitlesAndClosedCaptions.uplugin b/Engine/Plugins/Experimental/SubtitlesAndClosedCaptions/SubtitlesAndClosedCaptions.uplugin
--- a/Engine/Plugins/Experimental/SubtitlesAndClosedCaptions/SubtitlesAndClosedCaptions.uplugin
+++ b/Engine/Plugins/Experimental/SubtitlesAndClosedCaptions/SubtitlesAndClosedCaptions.uplugin
@@ -24,14 +24,6 @@
 				"Win64"
 			]
 		},
-		{
-			"Name": "SubtitlesAndClosedCaptionsTest",
-			"Type": "Runtime",
-			"LoadingPhase": "Default",
-			"PlatformAllowList": [
-				"Win64"
-			]
-		},
 		{
 			"Name": "SubtitlesAndClosedCaptionsEditor",
 			"Type": "UncookedOnly",
