Make delayed-start subtitles respect internal/external timing choices.

diff --git a/Engine/Plugins/Experimental/SubtitlesAndClosedCaptions/Source/SubtitlesAndClosedCaptions/Private/SubtitlesSubsystem.cpp b/Engine/Plugins/Experimental/SubtitlesAndClosedCaptions/Source/SubtitlesAndClosedCaptions/Private/SubtitlesSubsystem.cpp
--- a/Engine/Plugins/Experimental/SubtitlesAndClosedCaptions/Source/SubtitlesAndClosedCaptions/Private/SubtitlesSubsystem.cpp
+++ b/Engine/Plugins/Experimental/SubtitlesAndClosedCaptions/Source/SubtitlesAndClosedCaptions/Private/SubtitlesSubsystem.cpp
@@ -20,7 +20,7 @@
 #include "UObject/StrongObjectPtrTemplates.h"
 #include "UObject/WeakObjectPtrTemplates.h"
 
-// #SUBTITLES_PRD / #SUBTITLES_PRD_TODO - remove subtitles development comments prior to release
+static constexpr float InfiniteDuration = FLT_MAX;
 
 #include UE_INLINE_GENERATED_CPP_BY_NAME(SubtitlesSubsystem)
 
@@ -43,7 +43,6 @@ void USubtitlesSubsystem::BindDelegates()
 
 void USubtitlesSubsystem::QueueSubtitle(const FQueueSubtitleParameters& Params, const ESubtitleTiming Timing)
 {
-	constexpr float InfiniteDuration = FLT_MAX;
 	const USubtitleAssetUserData& Subtitle = *CastChecked<const USubtitleAssetUserData>(&Params.Subtitle);
 	UE_LOG(LogSubtitlesAndClosedCaptions, Display, TEXT("QueueSubtitle: '%s'"), *Subtitle.Text.ToString());
 
@@ -57,19 +56,19 @@ void USubtitlesSubsystem::QueueSubtitle(const FQueueSubtitleParameters& Params,
 
 	if (IsInGameThread())
 	{
-		AddActiveSubtitle(Subtitle, Duration);
+		AddActiveSubtitle(Subtitle, Duration, Timing);
 	}
 	else
 	{
 		ExecuteOnGameThread(
 			TEXT("USubtitlesSubsystem::HandleQueueSubtitle"),
-			[WeakThis = TWeakObjectPtr<USubtitlesSubsystem>(this), WeakSubtitle = TWeakObjectPtr<const USubtitleAssetUserData>(&Subtitle), Duration]()
+			[WeakThis = TWeakObjectPtr<USubtitlesSubsystem>(this), WeakSubtitle = TWeakObjectPtr<const USubtitleAssetUserData>(&Subtitle), Duration, Timing]()
 			{
 				if (TStrongObjectPtr<USubtitlesSubsystem> ThisPin = WeakThis.Pin())
 				{
 					if (TStrongObjectPtr<const USubtitleAssetUserData> SubtitlePin = WeakSubtitle.Pin())
 					{
-						ThisPin->AddActiveSubtitle(*SubtitlePin, Duration);
+						ThisPin->AddActiveSubtitle(*SubtitlePin, Duration, Timing);
 					}
 				}
 			}
@@ -77,7 +76,7 @@ void USubtitlesSubsystem::QueueSubtitle(const FQueueSubtitleParameters& Params,
 	}
 }
 
-void USubtitlesSubsystem::AddActiveSubtitle(const USubtitleAssetUserData& Subtitle, float Duration)
+void USubtitlesSubsystem::AddActiveSubtitle(const USubtitleAssetUserData& Subtitle, float Duration, const ESubtitleTiming Timing)
 {
 	// If the subtitle is already active then update its duration (by removing it then and then re-adding it)
 	const FActiveSubtitle* FoundActiveSubtitle = ActiveSubtitles.FindByPredicate(
@@ -94,7 +93,7 @@ void USubtitlesSubsystem::AddActiveSubtitle(const USubtitleAssetUserData& Subtit
 	{
 		// The timer handle will be reused for the duration when it enters the active subtitle queue.
 		// For now the timer tracks how long until it enters that queue.
-		FTimerDelegate Delegate = FTimerDelegate::CreateUObject(this, &USubtitlesSubsystem::MakeDelayedSubtitleActive, NewActiveSubtitle.Subtitle);
+		FTimerDelegate Delegate = FTimerDelegate::CreateUObject(this, &USubtitlesSubsystem::MakeDelayedSubtitleActive, NewActiveSubtitle.Subtitle, Timing);
 		GetWorldRef().GetTimerManager().SetTimer(NewActiveSubtitle.DurationTimerHandle, MoveTemp(Delegate), Subtitle.StartOffset, /*bLoop=*/false);
 
 		DelayedSubtitles.Add(MoveTemp(NewActiveSubtitle));
@@ -106,15 +105,12 @@ void USubtitlesSubsystem::AddActiveSubtitle(const USubtitleAssetUserData& Subtit
 		FTimerDelegate Delegate = FTimerDelegate::CreateUObject(this, &USubtitlesSubsystem::RemoveActiveSubtitle, NewActiveSubtitle.Subtitle);
 		Duration = FMath::Max(Duration, SubtitleMinDuration);
 		GetWorldRef().GetTimerManager().SetTimer(NewActiveSubtitle.DurationTimerHandle, MoveTemp(Delegate), Duration, /*bLoop=*/false);
-
-		ActiveSubtitles.Add(MoveTemp(NewActiveSubtitle));
-		ActiveSubtitles.StableSort([](const FActiveSubtitle& Lhs, const FActiveSubtitle& Rhs) { return Lhs.Subtitle->Priority > Rhs.Subtitle->Priority; });
-
-		UpdateWidgetData();
+		
+		AddAndDisplaySubtitle(NewActiveSubtitle);
 	}
 }
 
-void USubtitlesSubsystem::MakeDelayedSubtitleActive(TObjectPtr<const USubtitleAssetUserData> Subtitle)
+void USubtitlesSubsystem::MakeDelayedSubtitleActive(TObjectPtr<const USubtitleAssetUserData> Subtitle, const ESubtitleTiming Timing)
 {
 	if (IsValid(Subtitle))
 	{
@@ -125,21 +121,47 @@ void USubtitlesSubsystem::MakeDelayedSubtitleActive(TObjectPtr<const USubtitleAs
 			const USubtitleAssetUserData& SubtitleAsset = *DelayedSubtitle->Subtitle;
 			const float Duration = FMath::Max(SubtitleAsset.Duration, SubtitleMinDuration);
 
-			// Reuse the Timer Handle for duration, now that it's no longer needed for the delay.
-			FTimerDelegate Delegate = FTimerDelegate::CreateUObject(this, &USubtitlesSubsystem::RemoveActiveSubtitle, DelayedSubtitle->Subtitle);
-			GetWorldRef().GetTimerManager().SetTimer(DelayedSubtitle->DurationTimerHandle, MoveTemp(Delegate), Duration, /*bLoop=*/false);
+			if (Timing == ESubtitleTiming::InternallyTimed)
+			{
+				// Reuse the Timer Handle for duration, now that it's no longer needed for the delay.
+				FTimerDelegate Delegate = FTimerDelegate::CreateUObject(this, &USubtitlesSubsystem::RemoveActiveSubtitle, DelayedSubtitle->Subtitle);
+				GetWorldRef().GetTimerManager().SetTimer(DelayedSubtitle->DurationTimerHandle, MoveTemp(Delegate), Duration, /*bLoop=*/false);
+			}
 
 			// Insert the new subtitle to the actual queue and ensure it remains sorted by priority.
-			ActiveSubtitles.Add(MoveTemp(*DelayedSubtitle));
-			ActiveSubtitles.StableSort([](const FActiveSubtitle& Lhs, const FActiveSubtitle& Rhs) { return Lhs.Subtitle->Priority > Rhs.Subtitle->Priority; });
+			AddAndDisplaySubtitle(*DelayedSubtitle);
 
 			// Remove from the list of Delayed Subtitles
 			const int32 FirstRemovedIndex = Algo::RemoveIf(DelayedSubtitles, [&Subtitle](const FActiveSubtitle& DelayedSubtitle) {return DelayedSubtitle.Subtitle == Subtitle; });
 			DelayedSubtitles.SetNum(FirstRemovedIndex);
+		}
+	}
+}
 
-			UpdateWidgetData();
+void USubtitlesSubsystem::AddAndDisplaySubtitle(FActiveSubtitle& NewActiveSubtitle)
+{
+	check(NewActiveSubtitle.Subtitle); // Should only be called from AddActiveSubtitle or MakeDelayedSubtitle active, which would have a valid subtitle.
+	const USubtitleAssetUserData& NewSubtitle = *NewActiveSubtitle.Subtitle.Get();
+
+	bool ShouldAddSubtitle = true;
+	// Warn if the subtitle won't play due to being lower-priority than the currently-playing subtitle in that category.
+	for (const FActiveSubtitle PreQueuedSubtitle : ActiveSubtitles)
+	{
+		const USubtitleAssetUserData& PreQueuedSubtitleData = *PreQueuedSubtitle.Subtitle.Get();
+		if (PreQueuedSubtitleData.SubtitleType == NewSubtitle.SubtitleType && (PreQueuedSubtitleData.Priority >= NewSubtitle.Priority))
+		{
+			UE_LOG(LogSubtitlesAndClosedCaptions, Display, TEXT("Subtitle %s won't display: its priority is lower than or equal to the currently-playing subtitle in that category."), *NewSubtitle.Text.ToString());
+			ShouldAddSubtitle = false;
 		}
 	}
+
+	if (ShouldAddSubtitle)
+	{
+		ActiveSubtitles.Add(MoveTemp(NewActiveSubtitle));
+		ActiveSubtitles.StableSort([](const FActiveSubtitle& Lhs, const FActiveSubtitle& Rhs) { return Lhs.Subtitle->Priority > Rhs.Subtitle->Priority; });
+
+		UpdateWidgetData();
+	}
 }
 
 bool USubtitlesSubsystem::IsSubtitleActive(const UAssetUserData& Data) const
diff --git a/Engine/Plugins/Experimental/SubtitlesAndClosedCaptions/Source/SubtitlesAndClosedCaptions/Public/SubtitlesSubsystem.h b/Engine/Plugins/Experimental/SubtitlesAndClosedCaptions/Source/SubtitlesAndClosedCaptions/Public/SubtitlesSubsystem.h
--- a/Engine/Plugins/Experimental/SubtitlesAndClosedCaptions/Source/SubtitlesAndClosedCaptions/Public/SubtitlesSubsystem.h
+++ b/Engine/Plugins/Experimental/SubtitlesAndClosedCaptions/Source/SubtitlesAndClosedCaptions/Public/SubtitlesSubsystem.h
@@ -75,17 +75,19 @@ protected:
 	TArray<FActiveSubtitle> DelayedSubtitles;
 
 protected:
-	UE_API virtual void AddActiveSubtitle(const USubtitleAssetUserData& Subtitle, float Duration);
-	UE_API virtual void MakeDelayedSubtitleActive(TObjectPtr<const USubtitleAssetUserData> Subtitle);
+	UE_API virtual void AddActiveSubtitle(const USubtitleAssetUserData& Subtitle, float Duration, const ESubtitleTiming Timing);
+	UE_API virtual void MakeDelayedSubtitleActive(TObjectPtr<const USubtitleAssetUserData> Subtitle, const ESubtitleTiming Timing);
 	UE_API virtual void RemoveActiveSubtitle(TObjectPtr<const USubtitleAssetUserData> Subtitle);
 
+private:
+	UE_API void AddAndDisplaySubtitle(FActiveSubtitle& NewActiveSubtitle);
+
 #if WITH_DEV_AUTOMATION_TESTS
-#if 0 // Temporarily disabling these tests as they have a dangling reference that trips up the static analysis on certain build configurations.
-	friend struct SubtitlesAndClosedCaptions::Test::FMovieSceneSubtitlesTest;
-#endif
-	friend struct SubtitlesAndClosedCaptions::Test::FSubtitlesTest;
+	// These functions are for testing subtitle functionality and aren't intended for general use.
 public:
-	const TArray<FActiveSubtitle>& GetActiveSubtitles() const
+	// GetActiveSubtitles returns the sorted array of queued subtitles. Higher priorities are earlier in the array.
+	// For equally-prioritized subtitles, the oldest one is first in this array.
+	UE_API const TArray<FActiveSubtitle>& GetActiveSubtitles() const
 	{
 		return ActiveSubtitles;
 	}
@@ -102,7 +104,7 @@ public:
 	void TestActivatingDelayedSubtitle(const UAssetUserData& Data)
 	{
 		const USubtitleAssetUserData* Subtitle = CastChecked<const USubtitleAssetUserData>(&Data);
-		MakeDelayedSubtitleActive(Subtitle);
+		MakeDelayedSubtitleActive(Subtitle, ESubtitleTiming::InternallyTimed);
 	}
 #endif
 	//	WITH_DEV_AUTOMATION_TESTS
