Build fix for Linux Incremental build issue

diff --git a/Engine/Plugins/Experimental/SubtitlesAndClosedCaptions/Source/SubtitlesAndClosedCaptions/Private/SubtitlesSubsystem.cpp b/Engine/Plugins/Experimental/SubtitlesAndClosedCaptions/Source/SubtitlesAndClosedCaptions/Private/SubtitlesSubsystem.cpp
--- a/Engine/Plugins/Experimental/SubtitlesAndClosedCaptions/Source/SubtitlesAndClosedCaptions/Private/SubtitlesSubsystem.cpp
+++ b/Engine/Plugins/Experimental/SubtitlesAndClosedCaptions/Source/SubtitlesAndClosedCaptions/Private/SubtitlesSubsystem.cpp
@@ -145,7 +145,7 @@ void USubtitlesSubsystem::AddAndDisplaySubtitle(FActiveSubtitle& NewActiveSubtit
 
 	bool ShouldAddSubtitle = true;
 	// Warn if the subtitle won't play due to being lower-priority than the currently-playing subtitle in that category.
-	for (const FActiveSubtitle PreQueuedSubtitle : ActiveSubtitles)
+	for (const FActiveSubtitle& PreQueuedSubtitle : ActiveSubtitles)
 	{
 		const USubtitleAssetUserData& PreQueuedSubtitleData = *PreQueuedSubtitle.Subtitle.Get();
 		if (PreQueuedSubtitleData.SubtitleType == NewSubtitle.SubtitleType && (PreQueuedSubtitleData.Priority >= NewSubtitle.Priority))
