Separate the system for automatically queueing subtitles into its own subsystem to ensure the lifecycle is managed appropriately.

diff --git a/Engine/Plugins/Experimental/SubtitlesAndClosedCaptions/Source/SubtitlesAndClosedCaptions/Private/SubtitlesAudioSubsystem.cpp b/Engine/Plugins/Experimental/SubtitlesAndClosedCaptions/Source/SubtitlesAndClosedCaptions/Private/SubtitlesAudioSubsystem.cpp
new file mode 100644
--- /dev/null
+++ b/Engine/Plugins/Experimental/SubtitlesAndClosedCaptions/Source/SubtitlesAndClosedCaptions/Private/SubtitlesAudioSubsystem.cpp
@@ -0,0 +1,32 @@
+// Copyright Epic Games, Inc. All Rights Reserved.
+
+#include "SubtitlesAudioSubsystem.h"
+
+#include "ActiveSound.h"
+#include "Subtitles/SubtitlesAndClosedCaptionsDelegates.h"
+
+void USubtitlesAudioSubsystem::NotifyActiveSoundCreated(FActiveSound& ActiveSound)
+{
+	USoundBase* Sound = ActiveSound.GetSound();
+
+	if (IsValid(Sound))
+	{
+		if (const UAssetUserData * Subtitle{ Sound->GetAssetUserDataOfClass(USubtitleAssetUserData::StaticClass()) })
+		{
+			FSubtitlesAndClosedCaptionsDelegates::QueueSubtitle.ExecuteIfBound(FQueueSubtitleParameters{ *Subtitle }, ESubtitleTiming::InternallyTimed);
+		}
+	}
+}
+
+void USubtitlesAudioSubsystem::NotifyActiveSoundDeleting(const FActiveSound& ActiveSound)
+{
+	USoundBase* Sound = ActiveSound.GetSound();
+
+	if (IsValid(Sound))
+	{
+		if (const UAssetUserData * Subtitle{ Sound->GetAssetUserDataOfClass(USubtitleAssetUserData::StaticClass()) })
+		{
+			FSubtitlesAndClosedCaptionsDelegates::StopSubtitle.ExecuteIfBound(*Subtitle);
+		}
+	}
+}
\ No newline at end of file
diff --git a/Engine/Plugins/Experimental/SubtitlesAndClosedCaptions/Source/SubtitlesAndClosedCaptions/Public/SubtitlesAudioSubsystem.h b/Engine/Plugins/Experimental/SubtitlesAndClosedCaptions/Source/SubtitlesAndClosedCaptions/Public/SubtitlesAudioSubsystem.h
new file mode 100644
--- /dev/null
+++ b/Engine/Plugins/Experimental/SubtitlesAndClosedCaptions/Source/SubtitlesAndClosedCaptions/Public/SubtitlesAudioSubsystem.h
@@ -0,0 +1,30 @@
+// Copyright Epic Games, Inc. All Rights Reserved.
+
+#pragma once 
+#include "Subsystems/AudioEngineSubsystem.h"
+#include "ActiveSoundUpdateInterface.h"
+#include "SubtitlesAudioSubsystem.generated.h"
+
+#define UE_API SUBTITLESANDCLOSEDCAPTIONS_API
+
+/*
+*	SubtitlesAudioSubsystem - AudioEngineSubsystem for automatically queueing subtitles on audio classes.
+*
+*	Separated from SubtitlesSubsystem, which is a UWorldSubsystem, to avoid issues with lifecycle management
+*	as some parts need access to the world and others to the audio device.
+*/
+UCLASS(MinimalAPI, config = Game, defaultConfig)
+class USubtitlesAudioSubsystem : public UAudioEngineSubsystem
+	, public IActiveSoundUpdateInterface
+{
+	GENERATED_BODY()
+public:
+	USubtitlesAudioSubsystem() = default;
+
+	//~ Begin IActiveSoundUpdateInterface
+	virtual void NotifyActiveSoundCreated(FActiveSound& ActiveSound) override;
+	virtual void NotifyActiveSoundDeleting(const FActiveSound& ActiveSound) override;
+	//~ End IActiveSoundUpdateInterface
+};
+
+#undef UE_API
