Allow for new tracks to be null when adding in FSubtitlesTrackEditor::HandleAddMenuEntryExecute().

diff --git a/Engine/Plugins/Experimental/SubtitlesAndClosedCaptions/Source/SubtitlesAndClosedCaptionsEditor/Private/SubtitlesTrackEditor.cpp b/Engine/Plugins/Experimental/SubtitlesAndClosedCaptions/Source/SubtitlesAndClosedCaptionsEditor/Private/SubtitlesTrackEditor.cpp
--- a/Engine/Plugins/Experimental/SubtitlesAndClosedCaptions/Source/SubtitlesAndClosedCaptionsEditor/Private/SubtitlesTrackEditor.cpp
+++ b/Engine/Plugins/Experimental/SubtitlesAndClosedCaptions/Source/SubtitlesAndClosedCaptionsEditor/Private/SubtitlesTrackEditor.cpp
@@ -62,13 +62,14 @@ void FSubtitlesTrackEditor::HandleAddMenuEntryExecute()
 	FocusedMovieScene->Modify();
 
 	UMovieSceneSubtitlesTrack* NewTrack = FocusedMovieScene->AddTrack<UMovieSceneSubtitlesTrack>();
-	check(NewTrack);
-
-	NewTrack->SetDisplayName(LOCTEXT("SubtitlesTrackName", "Subtitles"));
-
-	if (GetSequencer().IsValid())
+	if (IsValid(NewTrack))
 	{
-		GetSequencer()->OnAddTrack(NewTrack, FGuid());
+		NewTrack->SetDisplayName(LOCTEXT("SubtitlesTrackName", "Subtitles"));
+
+		if (GetSequencer().IsValid())
+		{
+			GetSequencer()->OnAddTrack(NewTrack, FGuid());
+		}
 	}
 }
 
