Fix switching switching subtitle duration enum not respecting UseSoundDuration

diff --git a/Engine/Plugins/Experimental/SubtitlesAndClosedCaptions/Source/SubtitlesAndClosedCaptions/Private/SubtitlesSubsystem.cpp b/Engine/Plugins/Experimental/SubtitlesAndClosedCaptions/Source/SubtitlesAndClosedCaptions/Private/SubtitlesSubsystem.cpp
--- a/Engine/Plugins/Experimental/SubtitlesAndClosedCaptions/Source/SubtitlesAndClosedCaptions/Private/SubtitlesSubsystem.cpp
+++ b/Engine/Plugins/Experimental/SubtitlesAndClosedCaptions/Source/SubtitlesAndClosedCaptions/Private/SubtitlesSubsystem.cpp
@@ -100,11 +100,14 @@ void USubtitlesSubsystem::AddActiveSubtitle(const USubtitleAssetUserData& Subtit
 	}
 	else
 	{
-		// Without the delayed offset, instantly enter the queue as usual.
-		// The timer here tracks how long until the subtitle will expire and leave the active subtitle queue.
-		FTimerDelegate Delegate = FTimerDelegate::CreateUObject(this, &USubtitlesSubsystem::RemoveActiveSubtitle, NewActiveSubtitle.Subtitle);
-		Duration = FMath::Max(Duration, SubtitleMinDuration);
-		GetWorldRef().GetTimerManager().SetTimer(NewActiveSubtitle.DurationTimerHandle, MoveTemp(Delegate), Duration, /*bLoop=*/false);
+		if (Timing == ESubtitleTiming::InternallyTimed)
+		{
+			// Without the delayed offset, instantly enter the queue as usual.
+			// The timer here tracks how long until the subtitle will expire and leave the active subtitle queue.
+			FTimerDelegate Delegate = FTimerDelegate::CreateUObject(this, &USubtitlesSubsystem::RemoveActiveSubtitle, NewActiveSubtitle.Subtitle);
+			Duration = FMath::Max(Duration, SubtitleMinDuration);
+			GetWorldRef().GetTimerManager().SetTimer(NewActiveSubtitle.DurationTimerHandle, MoveTemp(Delegate), Duration, /*bLoop=*/false);
+		}
 		
 		AddAndDisplaySubtitle(NewActiveSubtitle);
 	}
