Multiple subtitles attached to SoundWaves can now all play, instead of just the first.

diff --git a/Engine/Plugins/Experimental/SubtitlesAndClosedCaptions/Source/SubtitlesAndClosedCaptions/Private/SubtitlesAudioSubsystem.cpp b/Engine/Plugins/Experimental/SubtitlesAndClosedCaptions/Source/SubtitlesAndClosedCaptions/Private/SubtitlesAudioSubsystem.cpp
--- a/Engine/Plugins/Experimental/SubtitlesAndClosedCaptions/Source/SubtitlesAndClosedCaptions/Private/SubtitlesAudioSubsystem.cpp
+++ b/Engine/Plugins/Experimental/SubtitlesAndClosedCaptions/Source/SubtitlesAndClosedCaptions/Private/SubtitlesAudioSubsystem.cpp
@@ -11,12 +11,16 @@ void USubtitlesAudioSubsystem::NotifyActiveSoundCreated(FActiveSound& ActiveSoun
 
 	if (IsValid(Sound))
 	{
-		if (const UAssetUserData * Subtitle{ Sound->GetAssetUserDataOfClass(USubtitleAssetUserData::StaticClass()) })
+		if (Sound->HasAssetUserDataOfClass(USubtitleAssetUserData::StaticClass()))
 		{
-			if (const USubtitleAssetUserData* SubtitleAsset = CastChecked<USubtitleAssetUserData>(Subtitle))	// Should never fail since Subtitle would only be valid if there's userdata of that class.
+			for (const TObjectPtr<UAssetUserData>& UserData: Sound->AssetUserData)
 			{
-				const bool bUseDurationProperty = (SubtitleAsset->SubtitleDurationType == ESubtitleDurationType::UseDurationProperty);
-				FSubtitlesAndClosedCaptionsDelegates::QueueSubtitle.ExecuteIfBound(FQueueSubtitleParameters{ *Subtitle }, bUseDurationProperty ? ESubtitleTiming::InternallyTimed : ESubtitleTiming::ExternallyTimed);
+				check(IsValid(UserData));
+				if (const USubtitleAssetUserData* Subtitle = Cast<USubtitleAssetUserData>(UserData.Get()))
+				{
+					const bool bUseDurationProperty = (Subtitle->SubtitleDurationType == ESubtitleDurationType::UseDurationProperty);
+					FSubtitlesAndClosedCaptionsDelegates::QueueSubtitle.ExecuteIfBound(FQueueSubtitleParameters{ *Subtitle }, bUseDurationProperty ? ESubtitleTiming::InternallyTimed : ESubtitleTiming::ExternallyTimed);
+				}
 			}
 		}
 	}
@@ -28,13 +32,17 @@ void USubtitlesAudioSubsystem::NotifyActiveSoundDeleting(const FActiveSound& Act
 
 	if (IsValid(Sound))
 	{
-		if (const UAssetUserData * Subtitle{ Sound->GetAssetUserDataOfClass(USubtitleAssetUserData::StaticClass()) })
+		if (Sound->HasAssetUserDataOfClass(USubtitleAssetUserData::StaticClass()))
 		{
-			if (const USubtitleAssetUserData* SubtitleAsset = CastChecked<USubtitleAssetUserData>(Subtitle))	// Should never fail since Subtitle would only be valid if there's userdata of that class.
+			for (const TObjectPtr<UAssetUserData>& UserData : Sound->AssetUserData)
 			{
-				if (SubtitleAsset->SubtitleDurationType == ESubtitleDurationType::UseSoundDuration)
+				check(IsValid(UserData));
+				if (const USubtitleAssetUserData* Subtitle = Cast<USubtitleAssetUserData>(UserData.Get()))
 				{
-					FSubtitlesAndClosedCaptionsDelegates::StopSubtitle.ExecuteIfBound(*Subtitle);
+					if (Subtitle->SubtitleDurationType == ESubtitleDurationType::UseSoundDuration)
+					{
+						FSubtitlesAndClosedCaptionsDelegates::StopSubtitle.ExecuteIfBound(*Subtitle);
+					}
 				}
 			}
 		}
