Don't queue subtitles if the engine has disabled them

diff --git a/Engine/Plugins/Experimental/SubtitlesAndClosedCaptions/Source/SubtitlesAndClosedCaptions/Private/SubtitlesSubsystem.cpp b/Engine/Plugins/Experimental/SubtitlesAndClosedCaptions/Source/SubtitlesAndClosedCaptions/Private/SubtitlesSubsystem.cpp
--- a/Engine/Plugins/Experimental/SubtitlesAndClosedCaptions/Source/SubtitlesAndClosedCaptions/Private/SubtitlesSubsystem.cpp
+++ b/Engine/Plugins/Experimental/SubtitlesAndClosedCaptions/Source/SubtitlesAndClosedCaptions/Private/SubtitlesSubsystem.cpp
@@ -9,6 +9,7 @@
 #include "Containers/Ticker.h"
 #include "Engine/Font.h"
 #include "Engine/World.h"
+#include "Engine.h"
 #include "Math/UnrealMathUtility.h"
 #include "Styling/CoreStyle.h"
 #include "SubtitlesAndClosedCaptionsModule.h"
@@ -43,6 +44,17 @@ void USubtitlesSubsystem::BindDelegates()
 
 void USubtitlesSubsystem::QueueSubtitle(const FQueueSubtitleParameters& Params, const ESubtitleTiming Timing)
 {
+	// Don't queue subtitles when the engine doesn't have them enabled or forced off. 
+	if (GEngine == nullptr)
+	{
+		UE_LOG(LogSubtitlesAndClosedCaptions, Warning, TEXT("Couldn't check if subtitles are enabled, because GEngine nullptr."));
+		return;
+	}
+	if(!GEngine->bSubtitlesEnabled || GEngine->bSubtitlesForcedOff)
+	{
+		return;
+	}
+
 	const USubtitleAssetUserData& Subtitle = *CastChecked<const USubtitleAssetUserData>(&Params.Subtitle);
 	UE_LOG(LogSubtitlesAndClosedCaptions, Display, TEXT("QueueSubtitle: '%s'"), *Subtitle.Text.ToString());
 
@@ -100,14 +112,11 @@ void USubtitlesSubsystem::AddActiveSubtitle(const USubtitleAssetUserData& Subtit
 	}
 	else
 	{
-		if (Timing == ESubtitleTiming::InternallyTimed)
-		{
-			// Without the delayed offset, instantly enter the queue as usual.
-			// The timer here tracks how long until the subtitle will expire and leave the active subtitle queue.
-			FTimerDelegate Delegate = FTimerDelegate::CreateUObject(this, &USubtitlesSubsystem::RemoveActiveSubtitle, NewActiveSubtitle.Subtitle);
-			Duration = FMath::Max(Duration, SubtitleMinDuration);
-			GetWorldRef().GetTimerManager().SetTimer(NewActiveSubtitle.DurationTimerHandle, MoveTemp(Delegate), Duration, /*bLoop=*/false);
-		}
+		// Without the delayed offset, instantly enter the queue as usual.
+		// The timer here tracks how long until the subtitle will expire and leave the active subtitle queue.
+		FTimerDelegate Delegate = FTimerDelegate::CreateUObject(this, &USubtitlesSubsystem::RemoveActiveSubtitle, NewActiveSubtitle.Subtitle);
+		Duration = FMath::Max(Duration, SubtitleMinDuration);
+		GetWorldRef().GetTimerManager().SetTimer(NewActiveSubtitle.DurationTimerHandle, MoveTemp(Delegate), Duration, /*bLoop=*/false);
 		
 		AddAndDisplaySubtitle(NewActiveSubtitle);
 	}
