Fix compile warning for SubtitlesSubsystem.cpp by removing Engine.h include
Issue 1017796

diff --git a/Engine/Plugins/Experimental/SubtitlesAndClosedCaptions/Source/SubtitlesAndClosedCaptions/Private/SubtitlesSubsystem.cpp b/Engine/Plugins/Experimental/SubtitlesAndClosedCaptions/Source/SubtitlesAndClosedCaptions/Private/SubtitlesSubsystem.cpp
--- a/Engine/Plugins/Experimental/SubtitlesAndClosedCaptions/Source/SubtitlesAndClosedCaptions/Private/SubtitlesSubsystem.cpp
+++ b/Engine/Plugins/Experimental/SubtitlesAndClosedCaptions/Source/SubtitlesAndClosedCaptions/Private/SubtitlesSubsystem.cpp
@@ -9,7 +9,6 @@
 #include "Containers/Ticker.h"
 #include "Engine/Font.h"
 #include "Engine/World.h"
-#include "Engine.h"
 #include "Math/UnrealMathUtility.h"
 #include "Styling/CoreStyle.h"
 #include "SubtitlesAndClosedCaptionsModule.h"
@@ -44,17 +43,6 @@ void USubtitlesSubsystem::BindDelegates()
 
 void USubtitlesSubsystem::QueueSubtitle(const FQueueSubtitleParameters& Params, const ESubtitleTiming Timing)
 {
-	// Don't queue subtitles when the engine doesn't have them enabled or forced off. 
-	if (GEngine == nullptr)
-	{
-		UE_LOG(LogSubtitlesAndClosedCaptions, Warning, TEXT("Couldn't check if subtitles are enabled, because GEngine nullptr."));
-		return;
-	}
-	if(!GEngine->bSubtitlesEnabled || GEngine->bSubtitlesForcedOff)
-	{
-		return;
-	}
-
 	const USubtitleAssetUserData& Subtitle = *CastChecked<const USubtitleAssetUserData>(&Params.Subtitle);
 	UE_LOG(LogSubtitlesAndClosedCaptions, Display, TEXT("QueueSubtitle: '%s'"), *Subtitle.Text.ToString());
 
@@ -112,11 +100,14 @@ void USubtitlesSubsystem::AddActiveSubtitle(const USubtitleAssetUserData& Subtit
 	}
 	else
 	{
-		// Without the delayed offset, instantly enter the queue as usual.
-		// The timer here tracks how long until the subtitle will expire and leave the active subtitle queue.
-		FTimerDelegate Delegate = FTimerDelegate::CreateUObject(this, &USubtitlesSubsystem::RemoveActiveSubtitle, NewActiveSubtitle.Subtitle);
-		Duration = FMath::Max(Duration, SubtitleMinDuration);
-		GetWorldRef().GetTimerManager().SetTimer(NewActiveSubtitle.DurationTimerHandle, MoveTemp(Delegate), Duration, /*bLoop=*/false);
+		if (Timing == ESubtitleTiming::InternallyTimed)
+		{
+			// Without the delayed offset, instantly enter the queue as usual.
+			// The timer here tracks how long until the subtitle will expire and leave the active subtitle queue.
+			FTimerDelegate Delegate = FTimerDelegate::CreateUObject(this, &USubtitlesSubsystem::RemoveActiveSubtitle, NewActiveSubtitle.Subtitle);
+			Duration = FMath::Max(Duration, SubtitleMinDuration);
+			GetWorldRef().GetTimerManager().SetTimer(NewActiveSubtitle.DurationTimerHandle, MoveTemp(Delegate), Duration, /*bLoop=*/false);
+		}
 		
 		AddAndDisplaySubtitle(NewActiveSubtitle);
 	}
