Makes sequencer subtitles' visibility correctly update when scrubbing along PIE.

diff --git a/Engine/Plugins/Experimental/SubtitlesAndClosedCaptions/Source/SubtitlesAndClosedCaptions/Public/MovieSceneSubtitlesSystem.h b/Engine/Plugins/Experimental/SubtitlesAndClosedCaptions/Source/SubtitlesAndClosedCaptions/Public/MovieSceneSubtitlesSystem.h
--- a/Engine/Plugins/Experimental/SubtitlesAndClosedCaptions/Source/SubtitlesAndClosedCaptions/Public/MovieSceneSubtitlesSystem.h
+++ b/Engine/Plugins/Experimental/SubtitlesAndClosedCaptions/Source/SubtitlesAndClosedCaptions/Public/MovieSceneSubtitlesSystem.h
@@ -99,8 +99,11 @@ struct FEvaluateSubtitles
 				// If not already active, queue the subtitle. It will remain visible until manually stopped; if the frame is frozen on a subtitle then display that subtitle forever.
 				if (LastEvaluatedFrame >= FirstSubtitleFrame && LastEvaluatedFrame < LastSubtitleFrame)
 				{
+					// This code is only run when over a subtitle section, so it can't reliably stop subtitles when the playback position jumps or scrubs.
+					// To avoid getting stuck on in those conditions, queue in quarter-second chunks that will stop on their own if not done manually.
 					FQueueSubtitleParameters Params{ AssetUserData };
-					FSubtitlesAndClosedCaptionsDelegates::QueueSubtitle.ExecuteIfBound(Params, ESubtitleTiming::ExternallyTimed);
+					Params.Duration = 0.25f;
+					FSubtitlesAndClosedCaptionsDelegates::QueueSubtitle.ExecuteIfBound(Params, ESubtitleTiming::InternallyTimed);
 				}
 			}
 			else if(LastSubtitleFrame <= LastEvaluatedFrame)
