Reinstate change to respect Engine subtitle enabled-ness with corrected Engine.h #include

diff --git a/Engine/Plugins/Experimental/SubtitlesAndClosedCaptions/Source/SubtitlesAndClosedCaptions/Private/SubtitlesSubsystem.cpp b/Engine/Plugins/Experimental/SubtitlesAndClosedCaptions/Source/SubtitlesAndClosedCaptions/Private/SubtitlesSubsystem.cpp
--- a/Engine/Plugins/Experimental/SubtitlesAndClosedCaptions/Source/SubtitlesAndClosedCaptions/Private/SubtitlesSubsystem.cpp
+++ b/Engine/Plugins/Experimental/SubtitlesAndClosedCaptions/Source/SubtitlesAndClosedCaptions/Private/SubtitlesSubsystem.cpp
@@ -9,6 +9,7 @@
 #include "Containers/Ticker.h"
 #include "Engine/Font.h"
 #include "Engine/World.h"
+#include "Engine/Engine.h"
 #include "Math/UnrealMathUtility.h"
 #include "Styling/CoreStyle.h"
 #include "SubtitlesAndClosedCaptionsModule.h"
@@ -43,6 +44,17 @@ void USubtitlesSubsystem::BindDelegates()
 
 void USubtitlesSubsystem::QueueSubtitle(const FQueueSubtitleParameters& Params, const ESubtitleTiming Timing)
 {
+	// Don't queue subtitles when the engine doesn't have them enabled or forced off. 
+	if (GEngine == nullptr)
+	{
+		UE_LOG(LogSubtitlesAndClosedCaptions, Warning, TEXT("Couldn't check if subtitles are enabled, because GEngine nullptr."));
+		return;
+	}
+	if (!GEngine->bSubtitlesEnabled || GEngine->bSubtitlesForcedOff)
+	{
+		return;
+	}
+
 	const USubtitleAssetUserData& Subtitle = *CastChecked<const USubtitleAssetUserData>(&Params.Subtitle);
 	UE_LOG(LogSubtitlesAndClosedCaptions, Display, TEXT("QueueSubtitle: '%s'"), *Subtitle.Text.ToString());
 
