Fix scrubbed subtitles in the sequencer flickering

diff --git a/Engine/Plugins/Experimental/SubtitlesAndClosedCaptions/Source/SubtitlesAndClosedCaptions/Public/MovieSceneSubtitlesSystem.h b/Engine/Plugins/Experimental/SubtitlesAndClosedCaptions/Source/SubtitlesAndClosedCaptions/Public/MovieSceneSubtitlesSystem.h
--- a/Engine/Plugins/Experimental/SubtitlesAndClosedCaptions/Source/SubtitlesAndClosedCaptions/Public/MovieSceneSubtitlesSystem.h
+++ b/Engine/Plugins/Experimental/SubtitlesAndClosedCaptions/Source/SubtitlesAndClosedCaptions/Public/MovieSceneSubtitlesSystem.h
@@ -94,19 +94,15 @@ struct FEvaluateSubtitles
 			const FFrameNumber FirstSubtitleFrame = SubtitleSection->GetRange().GetLowerBoundValue();
 			const FFrameNumber LastEvaluatedFrame = Context.GetFrameNumberRange().GetUpperBoundValue();
 
-			if (!FSubtitlesAndClosedCaptionsDelegates::IsSubtitleActive.Execute(AssetUserData))
+			// Queue the subtitle to display in short increments whenever playback is over the subtitle.
+			// This allows for jumping and scrubbing the sequence without the subtitle getting stuck on or off.
+			if (LastEvaluatedFrame >= FirstSubtitleFrame && LastEvaluatedFrame < LastSubtitleFrame)
 			{
-				// If not already active, queue the subtitle. It will remain visible until manually stopped; if the frame is frozen on a subtitle then display that subtitle forever.
-				if (LastEvaluatedFrame >= FirstSubtitleFrame && LastEvaluatedFrame < LastSubtitleFrame)
-				{
-					// This code is only run when over a subtitle section, so it can't reliably stop subtitles when the playback position jumps or scrubs.
-					// To avoid getting stuck on in those conditions, queue in quarter-second chunks that will stop on their own if not done manually.
-					FQueueSubtitleParameters Params{ AssetUserData };
-					Params.Duration = 0.25f;
-					FSubtitlesAndClosedCaptionsDelegates::QueueSubtitle.ExecuteIfBound(Params, ESubtitleTiming::InternallyTimed);
-				}
+				FQueueSubtitleParameters Params{ AssetUserData };
+				Params.Duration = 0.25f;
+				FSubtitlesAndClosedCaptionsDelegates::QueueSubtitle.ExecuteIfBound(Params, ESubtitleTiming::InternallyTimed);
 			}
-			else if(LastSubtitleFrame <= LastEvaluatedFrame)
+			else if(FSubtitlesAndClosedCaptionsDelegates::IsSubtitleActive.Execute(AssetUserData) && (LastSubtitleFrame <= LastEvaluatedFrame))
 			{
 				// If the subtitle is already active and outside the range of frames it should display for, remove it.
 				FSubtitlesAndClosedCaptionsDelegates::StopSubtitle.ExecuteIfBound(AssetUserData);
