Refactor SubtitleAssetUserData so subtitles are in an array inside. Merge to Release-38.
#virtualized

diff --git a/Engine/Plugins/Experimental/SubtitlesAndClosedCaptions/Source/SubtitlesAndClosedCaptions/Private/MovieSceneSubtitleSection.cpp b/Engine/Plugins/Experimental/SubtitlesAndClosedCaptions/Source/SubtitlesAndClosedCaptions/Private/MovieSceneSubtitleSection.cpp
--- a/Engine/Plugins/Experimental/SubtitlesAndClosedCaptions/Source/SubtitlesAndClosedCaptions/Private/MovieSceneSubtitleSection.cpp
+++ b/Engine/Plugins/Experimental/SubtitlesAndClosedCaptions/Source/SubtitlesAndClosedCaptions/Private/MovieSceneSubtitleSection.cpp
@@ -26,8 +26,9 @@ TOptional<TRange<FFrameNumber>> UMovieSceneSubtitleSection::GetAutoSizeRange() c
 		return TRange<FFrameNumber>();
 	}
 
+	// Find the maximum duration of all the subtitles in the UserData, including their delayed start offsets.
 	const FFrameRate FrameRate = GetTypedOuter<UMovieScene>()->GetTickResolution();
-	const FFrameTime DurationToUse = Subtitle->Duration * FrameRate;
+	const FFrameTime DurationToUse = Subtitle->GetMaximumDuration() * FrameRate;
 	return TRange<FFrameNumber>(GetInclusiveStartFrame(), GetInclusiveStartFrame() + DurationToUse.FrameNumber);
 }
 
diff --git a/Engine/Plugins/Experimental/SubtitlesAndClosedCaptions/Source/SubtitlesAndClosedCaptions/Private/MovieSceneSubtitlesTrack.cpp b/Engine/Plugins/Experimental/SubtitlesAndClosedCaptions/Source/SubtitlesAndClosedCaptions/Private/MovieSceneSubtitlesTrack.cpp
--- a/Engine/Plugins/Experimental/SubtitlesAndClosedCaptions/Source/SubtitlesAndClosedCaptions/Private/MovieSceneSubtitlesTrack.cpp
+++ b/Engine/Plugins/Experimental/SubtitlesAndClosedCaptions/Source/SubtitlesAndClosedCaptions/Private/MovieSceneSubtitlesTrack.cpp
@@ -27,7 +27,9 @@ UMovieSceneSection* UMovieSceneSubtitlesTrack::AddNewSubtitleOnRow(const USubtit
 	UMovieSceneSubtitleSection* NewSection = Cast<UMovieSceneSubtitleSection>(CreateNewSection());
 
 	const FFrameRate FrameRate = GetTypedOuter<UMovieScene>()->GetTickResolution();
-	const FFrameTime DurationToUse = Subtitle.Duration * FrameRate;
+
+	// Find the maximum duration of all the subtitles in the UserData, including their delayed start offsets.
+	const FFrameTime DurationToUse = Subtitle.GetMaximumDuration() * FrameRate;
 
 	ensure(DurationToUse.FrameNumber.Value > 0);
 	NewSection->InitialPlacementOnRow(SubtitleSections, InStartTime, DurationToUse.FrameNumber.Value, InRowIndex);
diff --git a/Engine/Plugins/Experimental/SubtitlesAndClosedCaptions/Source/SubtitlesAndClosedCaptions/Private/SubtitleTextBlock.cpp b/Engine/Plugins/Experimental/SubtitlesAndClosedCaptions/Source/SubtitlesAndClosedCaptions/Private/SubtitleTextBlock.cpp
--- a/Engine/Plugins/Experimental/SubtitlesAndClosedCaptions/Source/SubtitlesAndClosedCaptions/Private/SubtitleTextBlock.cpp
+++ b/Engine/Plugins/Experimental/SubtitlesAndClosedCaptions/Source/SubtitlesAndClosedCaptions/Private/SubtitleTextBlock.cpp
@@ -27,7 +27,7 @@ void USubtitleTextBlock::NativeConstruct()
 	}
 }
 
-void USubtitleTextBlock::StartDisplayingSubtitle(const USubtitleAssetUserData& Subtitle)
+void USubtitleTextBlock::StartDisplayingSubtitle(const FSubtitleAssetData& Subtitle)
 {
 	// Pick which UextBlock is relevant by category....
 	UTextBlock* TextToModify = nullptr;
diff --git a/Engine/Plugins/Experimental/SubtitlesAndClosedCaptions/Source/SubtitlesAndClosedCaptions/Private/SubtitlesAudioSubsystem.cpp b/Engine/Plugins/Experimental/SubtitlesAndClosedCaptions/Source/SubtitlesAndClosedCaptions/Private/SubtitlesAudioSubsystem.cpp
--- a/Engine/Plugins/Experimental/SubtitlesAndClosedCaptions/Source/SubtitlesAndClosedCaptions/Private/SubtitlesAudioSubsystem.cpp
+++ b/Engine/Plugins/Experimental/SubtitlesAndClosedCaptions/Source/SubtitlesAndClosedCaptions/Private/SubtitlesAudioSubsystem.cpp
@@ -3,6 +3,7 @@
 #include "SubtitlesAudioSubsystem.h"
 
 #include "ActiveSound.h"
+#include "Sound/DialogueSoundWaveProxy.h"
 #include "Subtitles/SubtitlesAndClosedCaptionsDelegates.h"
 
 #include UE_INLINE_GENERATED_CPP_BY_NAME(SubtitlesAudioSubsystem)
@@ -15,14 +16,27 @@ void USubtitlesAudioSubsystem::NotifyActiveSoundCreated(FActiveSound& ActiveSoun
 	{
 		if (Sound->HasAssetUserDataOfClass(USubtitleAssetUserData::StaticClass()))
 		{
-			for (const TObjectPtr<UAssetUserData>& UserData: Sound->AssetUserData)
+			const USubtitleAssetUserData* Subtitles = CastChecked<USubtitleAssetUserData>(Sound->GetAssetUserDataOfClass(USubtitleAssetUserData::StaticClass()));
+			for (const FSubtitleAssetData& Subtitle : Subtitles->Subtitles)
 			{
-				check(IsValid(UserData));
-				if (const USubtitleAssetUserData* Subtitle = Cast<USubtitleAssetUserData>(UserData.Get()))
-				{
-					const bool bUseDurationProperty = (Subtitle->SubtitleDurationType == ESubtitleDurationType::UseDurationProperty);
-					FSubtitlesAndClosedCaptionsDelegates::QueueSubtitle.ExecuteIfBound(FQueueSubtitleParameters{ *Subtitle }, bUseDurationProperty ? ESubtitleTiming::InternallyTimed : ESubtitleTiming::ExternallyTimed);
-				}
+				const bool bUseDurationProperty = (Subtitle.SubtitleDurationType == ESubtitleDurationType::UseDurationProperty);
+				FSubtitlesAndClosedCaptionsDelegates::QueueSubtitle.ExecuteIfBound(FQueueSubtitleParameters{ Subtitle }, bUseDurationProperty ? ESubtitleTiming::InternallyTimed : ESubtitleTiming::ExternallyTimed);
+			}
+		}
+
+		// Also check for DialogueWaves with localized text
+		if (ActiveSound.bHandleSubtitles)
+		{
+			const UDialogueSoundWaveProxy* DialogueWave = Cast<UDialogueSoundWaveProxy>(ActiveSound.GetSound());
+			if (IsValid(DialogueWave))
+			{
+				FSubtitleAssetData SubtitleCue;
+				SubtitleCue.Text = DialogueWave->GetSubtitleText();
+				SubtitleCue.Priority = ActiveSound.SubtitlePriority;
+				SubtitleCue.Duration = DialogueWave->GetDuration();
+				SubtitleCue.StartOffset = ActiveSound.RequestedStartTime;
+
+				FSubtitlesAndClosedCaptionsDelegates::QueueSubtitle.ExecuteIfBound(FQueueSubtitleParameters{ SubtitleCue }, ESubtitleTiming::InternallyTimed);
 			}
 		}
 	}
@@ -36,14 +50,15 @@ void USubtitlesAudioSubsystem::NotifyActiveSoundDeleting(const FActiveSound& Act
 	{
 		if (Sound->HasAssetUserDataOfClass(USubtitleAssetUserData::StaticClass()))
 		{
-			for (const TObjectPtr<UAssetUserData>& UserData : Sound->AssetUserData)
+			if (Sound->HasAssetUserDataOfClass(USubtitleAssetUserData::StaticClass()))
 			{
-				check(IsValid(UserData));
-				if (const USubtitleAssetUserData* Subtitle = Cast<USubtitleAssetUserData>(UserData.Get()))
+				const USubtitleAssetUserData* Subtitles = CastChecked<USubtitleAssetUserData>(Sound->GetAssetUserDataOfClass(USubtitleAssetUserData::StaticClass()));
+				for (const FSubtitleAssetData& Subtitle : Subtitles->Subtitles)
 				{
-					if (Subtitle->SubtitleDurationType == ESubtitleDurationType::UseSoundDuration)
+					const bool bUseDurationProperty = (Subtitle.SubtitleDurationType == ESubtitleDurationType::UseDurationProperty);
+					if (!bUseDurationProperty)
 					{
-						FSubtitlesAndClosedCaptionsDelegates::StopSubtitle.ExecuteIfBound(*Subtitle);
+						FSubtitlesAndClosedCaptionsDelegates::StopSubtitle.ExecuteIfBound(Subtitle);
 					}
 				}
 			}
diff --git a/Engine/Plugins/Experimental/SubtitlesAndClosedCaptions/Source/SubtitlesAndClosedCaptions/Private/SubtitlesBlueprintFunctionLibrary.cpp b/Engine/Plugins/Experimental/SubtitlesAndClosedCaptions/Source/SubtitlesAndClosedCaptions/Private/SubtitlesBlueprintFunctionLibrary.cpp
--- a/Engine/Plugins/Experimental/SubtitlesAndClosedCaptions/Source/SubtitlesAndClosedCaptions/Private/SubtitlesBlueprintFunctionLibrary.cpp
+++ b/Engine/Plugins/Experimental/SubtitlesAndClosedCaptions/Source/SubtitlesAndClosedCaptions/Private/SubtitlesBlueprintFunctionLibrary.cpp
@@ -6,37 +6,24 @@
 
 #include UE_INLINE_GENERATED_CPP_BY_NAME(SubtitlesBlueprintFunctionLibrary)
 
-void USubtitlesBlueprintFunctionLibrary::QueueSubtitle(const USubtitleAssetUserData* Subtitle, const ESubtitleTiming Timing)
+void USubtitlesBlueprintFunctionLibrary::QueueSubtitle(const FSubtitleAssetData& Subtitle, const ESubtitleTiming Timing)
 {
-	if (!Subtitle)
-	{
-		return;
-	}
-
-	FQueueSubtitleParameters Params{ *CastChecked<const UAssetUserData>(Subtitle) };
+	FQueueSubtitleParameters Params{ Subtitle };
 	FSubtitlesAndClosedCaptionsDelegates::QueueSubtitle.ExecuteIfBound(Params, Timing);
 }
 
-bool USubtitlesBlueprintFunctionLibrary::IsSubtitleActive(const USubtitleAssetUserData* Subtitle)
+bool USubtitlesBlueprintFunctionLibrary::IsSubtitleActive(const FSubtitleAssetData& Subtitle)
 {
-	if (!Subtitle)
-	{
-		return false;
-	}
 	if (!FSubtitlesAndClosedCaptionsDelegates::IsSubtitleActive.IsBound())
 	{
 		return false;
 	}
-	return FSubtitlesAndClosedCaptionsDelegates::IsSubtitleActive.Execute(*CastChecked<const UAssetUserData>(Subtitle));
+	return FSubtitlesAndClosedCaptionsDelegates::IsSubtitleActive.Execute(Subtitle);
 }
 
-void USubtitlesBlueprintFunctionLibrary::StopSubtitle(const USubtitleAssetUserData* Subtitle)
+void USubtitlesBlueprintFunctionLibrary::StopSubtitle(const FSubtitleAssetData& Subtitle)
 {
-	if (!Subtitle)
-	{
-		return;
-	}
-	FSubtitlesAndClosedCaptionsDelegates::StopSubtitle.ExecuteIfBound(*CastChecked<const UAssetUserData>(Subtitle));
+	FSubtitlesAndClosedCaptionsDelegates::StopSubtitle.ExecuteIfBound(Subtitle);
 }
 
 void USubtitlesBlueprintFunctionLibrary::StopAllSubtitles()
diff --git a/Engine/Plugins/Experimental/SubtitlesAndClosedCaptions/Source/SubtitlesAndClosedCaptions/Private/SubtitlesSubsystem.cpp b/Engine/Plugins/Experimental/SubtitlesAndClosedCaptions/Source/SubtitlesAndClosedCaptions/Private/SubtitlesSubsystem.cpp
--- a/Engine/Plugins/Experimental/SubtitlesAndClosedCaptions/Source/SubtitlesAndClosedCaptions/Private/SubtitlesSubsystem.cpp
+++ b/Engine/Plugins/Experimental/SubtitlesAndClosedCaptions/Source/SubtitlesAndClosedCaptions/Private/SubtitlesSubsystem.cpp
@@ -55,7 +55,7 @@ void USubtitlesSubsystem::QueueSubtitle(const FQueueSubtitleParameters& Params,
 		return;
 	}
 
-	const USubtitleAssetUserData& Subtitle = *CastChecked<const USubtitleAssetUserData>(&Params.Subtitle);
+	const FSubtitleAssetData& Subtitle = Params.Subtitle;
 	UE_LOG(LogSubtitlesAndClosedCaptions, Display, TEXT("QueueSubtitle: '%s'"), *Subtitle.Text.ToString());
 
 
@@ -66,47 +66,47 @@ void USubtitlesSubsystem::QueueSubtitle(const FQueueSubtitleParameters& Params,
 		Duration = Params.Duration.IsSet() ? Params.Duration.GetValue() : Subtitle.Duration;
 	}
 
+	const float StartOffset = Params.StartOffset.IsSet() ? Params.StartOffset.GetValue() : Subtitle.StartOffset;
+
 	if (IsInGameThread())
 	{
-		AddActiveSubtitle(Subtitle, Duration, Timing);
+		AddActiveSubtitle(Subtitle, Duration, StartOffset, Timing);
 	}
 	else
 	{
 		ExecuteOnGameThread(
 			TEXT("USubtitlesSubsystem::HandleQueueSubtitle"),
-			[WeakThis = TWeakObjectPtr<USubtitlesSubsystem>(this), WeakSubtitle = TWeakObjectPtr<const USubtitleAssetUserData>(&Subtitle), Duration, Timing]()
+			[WeakThis = TWeakObjectPtr<USubtitlesSubsystem>(this), Subtitle, Duration, StartOffset, Timing]()
 			{
 				if (TStrongObjectPtr<USubtitlesSubsystem> ThisPin = WeakThis.Pin())
 				{
-					if (TStrongObjectPtr<const USubtitleAssetUserData> SubtitlePin = WeakSubtitle.Pin())
-					{
-						ThisPin->AddActiveSubtitle(*SubtitlePin, Duration, Timing);
-					}
+					ThisPin->AddActiveSubtitle(Subtitle, Duration, StartOffset, Timing);
 				}
 			}
 		);
 	}
 }
 
-void USubtitlesSubsystem::AddActiveSubtitle(const USubtitleAssetUserData& Subtitle, float Duration, const ESubtitleTiming Timing)
+void USubtitlesSubsystem::AddActiveSubtitle(const FSubtitleAssetData& Subtitle, float Duration, const float StartOffset, const ESubtitleTiming Timing)
 {
 	// If the subtitle is already active then update its duration (by removing it then and then re-adding it)
 	const FActiveSubtitle* FoundActiveSubtitle = ActiveSubtitles.FindByPredicate(
-		[SubtitlePtr = &Subtitle](const FActiveSubtitle& ActiveSubtitle) { return ActiveSubtitle.Subtitle == SubtitlePtr; });
+		[&Subtitle](const FActiveSubtitle& ActiveSubtitle) { return ActiveSubtitle.Subtitle == Subtitle; });
 	if (FoundActiveSubtitle != nullptr)
 	{
-		RemoveActiveSubtitle(&Subtitle);
+		RemoveActiveSubtitle(Subtitle);
 	}
 
-	FActiveSubtitle NewActiveSubtitle{ &Subtitle };
+	FActiveSubtitle NewActiveSubtitle{ Subtitle };
 
 	// Subtitles with delayed offset need a timer to await their entry in the queue.
-	if (Subtitle.StartOffset > 0.f)
+	if (StartOffset > 0.f)
 	{
 		// The timer handle will be reused for the duration when it enters the active subtitle queue.
 		// For now the timer tracks how long until it enters that queue.
-		FTimerDelegate Delegate = FTimerDelegate::CreateUObject(this, &USubtitlesSubsystem::MakeDelayedSubtitleActive, NewActiveSubtitle.Subtitle, Timing);
-		GetWorldRef().GetTimerManager().SetTimer(NewActiveSubtitle.DurationTimerHandle, MoveTemp(Delegate), Subtitle.StartOffset, /*bLoop=*/false);
+		FTimerDelegate Delegate;
+		Delegate.BindUFunction(this, FName(TEXT("MakeDelayedSubtitleActive")), NewActiveSubtitle.Subtitle, Timing);
+		GetWorldRef().GetTimerManager().SetTimer(NewActiveSubtitle.DurationTimerHandle, MoveTemp(Delegate), StartOffset, /*bLoop=*/false);
 
 		DelayedSubtitles.Add(MoveTemp(NewActiveSubtitle));
 	}
@@ -116,7 +116,8 @@ void USubtitlesSubsystem::AddActiveSubtitle(const USubtitleAssetUserData& Subtit
 		{
 			// Without the delayed offset, instantly enter the queue as usual.
 			// The timer here tracks how long until the subtitle will expire and leave the active subtitle queue.
-			FTimerDelegate Delegate = FTimerDelegate::CreateUObject(this, &USubtitlesSubsystem::RemoveActiveSubtitle, NewActiveSubtitle.Subtitle);
+			FTimerDelegate Delegate;
+			Delegate.BindUFunction(this, FName(TEXT("RemoveActiveSubtitle")), NewActiveSubtitle.Subtitle);
 			Duration = FMath::Max(Duration, SubtitleMinDuration);
 			GetWorldRef().GetTimerManager().SetTimer(NewActiveSubtitle.DurationTimerHandle, MoveTemp(Delegate), Duration, /*bLoop=*/false);
 		}
@@ -125,44 +126,41 @@ void USubtitlesSubsystem::AddActiveSubtitle(const USubtitleAssetUserData& Subtit
 	}
 }
 
-void USubtitlesSubsystem::MakeDelayedSubtitleActive(TObjectPtr<const USubtitleAssetUserData> Subtitle, const ESubtitleTiming Timing)
+void USubtitlesSubsystem::MakeDelayedSubtitleActive(const FSubtitleAssetData& Subtitle, const ESubtitleTiming Timing)
 {
-	if (IsValid(Subtitle))
+	FActiveSubtitle* DelayedSubtitle = Algo::FindByPredicate(DelayedSubtitles, [&Subtitle](const FActiveSubtitle& DelayedSubtitle) {return DelayedSubtitle.Subtitle == Subtitle; });
+
+	if (DelayedSubtitle != nullptr)
 	{
-		FActiveSubtitle* DelayedSubtitle = Algo::FindByPredicate(DelayedSubtitles, [&Subtitle](const FActiveSubtitle& DelayedSubtitle) {return DelayedSubtitle.Subtitle == Subtitle; });
+		const FSubtitleAssetData& SubtitleAsset = DelayedSubtitle->Subtitle;
+		const float Duration = FMath::Max(SubtitleAsset.Duration, SubtitleMinDuration);
 
-		if (DelayedSubtitle != nullptr)
+		if (Timing == ESubtitleTiming::InternallyTimed)
 		{
-			const USubtitleAssetUserData& SubtitleAsset = *DelayedSubtitle->Subtitle;
-			const float Duration = FMath::Max(SubtitleAsset.Duration, SubtitleMinDuration);
-
-			if (Timing == ESubtitleTiming::InternallyTimed)
-			{
-				// Reuse the Timer Handle for duration, now that it's no longer needed for the delay.
-				FTimerDelegate Delegate = FTimerDelegate::CreateUObject(this, &USubtitlesSubsystem::RemoveActiveSubtitle, DelayedSubtitle->Subtitle);
-				GetWorldRef().GetTimerManager().SetTimer(DelayedSubtitle->DurationTimerHandle, MoveTemp(Delegate), Duration, /*bLoop=*/false);
-			}
+			// Reuse the Timer Handle for duration, now that it's no longer needed for the delay.
+			FTimerDelegate Delegate;
+			Delegate.BindUFunction(this, FName(TEXT("RemoveActiveSubtitle")), DelayedSubtitle->Subtitle);
+			GetWorldRef().GetTimerManager().SetTimer(DelayedSubtitle->DurationTimerHandle, MoveTemp(Delegate), Duration, /*bLoop=*/false);
+		}
 
-			// Insert the new subtitle to the actual queue and ensure it remains sorted by priority.
-			AddAndDisplaySubtitle(*DelayedSubtitle);
+		// Insert the new subtitle to the actual queue and ensure it remains sorted by priority.
+		AddAndDisplaySubtitle(*DelayedSubtitle);
 
-			// Remove from the list of Delayed Subtitles
-			const int32 FirstRemovedIndex = Algo::RemoveIf(DelayedSubtitles, [&Subtitle](const FActiveSubtitle& DelayedSubtitle) {return DelayedSubtitle.Subtitle == Subtitle; });
-			DelayedSubtitles.SetNum(FirstRemovedIndex);
-		}
+		// Remove from the list of Delayed Subtitles
+		const int32 FirstRemovedIndex = Algo::RemoveIf(DelayedSubtitles, [&Subtitle](const FActiveSubtitle& DelayedSubtitle) {return DelayedSubtitle.Subtitle == Subtitle; });
+		DelayedSubtitles.SetNum(FirstRemovedIndex);
 	}
 }
 
 void USubtitlesSubsystem::AddAndDisplaySubtitle(FActiveSubtitle& NewActiveSubtitle)
 {
-	check(NewActiveSubtitle.Subtitle); // Should only be called from AddActiveSubtitle or MakeDelayedSubtitle active, which would have a valid subtitle.
-	const USubtitleAssetUserData& NewSubtitle = *NewActiveSubtitle.Subtitle.Get();
+	const FSubtitleAssetData& NewSubtitle = NewActiveSubtitle.Subtitle;
 
 	bool ShouldAddSubtitle = true;
 	// Warn if the subtitle won't play due to being lower-priority than the currently-playing subtitle in that category.
 	for (const FActiveSubtitle& PreQueuedSubtitle : ActiveSubtitles)
 	{
-		const USubtitleAssetUserData& PreQueuedSubtitleData = *PreQueuedSubtitle.Subtitle.Get();
+		const FSubtitleAssetData& PreQueuedSubtitleData = PreQueuedSubtitle.Subtitle;
 		if (PreQueuedSubtitleData.SubtitleType == NewSubtitle.SubtitleType && (PreQueuedSubtitleData.Priority >= NewSubtitle.Priority))
 		{
 			UE_LOG(LogSubtitlesAndClosedCaptions, Display, TEXT("Subtitle %s won't display: its priority is lower than or equal to the currently-playing subtitle in that category."), *NewSubtitle.Text.ToString());
@@ -173,31 +171,28 @@ void USubtitlesSubsystem::AddAndDisplaySubtitle(FActiveSubtitle& NewActiveSubtit
 	if (ShouldAddSubtitle)
 	{
 		ActiveSubtitles.Add(MoveTemp(NewActiveSubtitle));
-		ActiveSubtitles.StableSort([](const FActiveSubtitle& Lhs, const FActiveSubtitle& Rhs) { return Lhs.Subtitle->Priority > Rhs.Subtitle->Priority; });
+		ActiveSubtitles.StableSort([](const FActiveSubtitle& Lhs, const FActiveSubtitle& Rhs) { return Lhs.Subtitle.Priority > Rhs.Subtitle.Priority; });
 
 		UpdateWidgetData();
 	}
 }
 
-bool USubtitlesSubsystem::IsSubtitleActive(const UAssetUserData& Data) const
+bool USubtitlesSubsystem::IsSubtitleActive(const FSubtitleAssetData& Data) const
 {
 	if (!ensureMsgf(IsInGameThread(), TEXT("IsSubtitleActive must currently be run on the GameThread - ActiveSubtitles vector is not locked")))
 	{
 		return false;
 	}
 
-	const USubtitleAssetUserData* Subtitle = CastChecked<const USubtitleAssetUserData>(&Data);
-
 	const FActiveSubtitle* FoundActiveSubtitle = ActiveSubtitles.FindByPredicate(
-		[Subtitle](const FActiveSubtitle& ActiveSubtitle) { return ActiveSubtitle.Subtitle == Subtitle; });
+		[Data](const FActiveSubtitle& ActiveSubtitle) { return ActiveSubtitle.Subtitle == Data; });
 
 	return FoundActiveSubtitle != nullptr;
 }
 
-void USubtitlesSubsystem::StopSubtitle(const UAssetUserData& Data)
+void USubtitlesSubsystem::StopSubtitle(const FSubtitleAssetData& Data)
 {
-	const USubtitleAssetUserData* Subtitle = CastChecked<const USubtitleAssetUserData>(&Data);
-	RemoveActiveSubtitle(Subtitle);
+	RemoveActiveSubtitle(Data);
 }
 
 void USubtitlesSubsystem::StopAllSubtitles()
@@ -230,7 +225,7 @@ void USubtitlesSubsystem::StopAllSubtitles()
 	}
 }
 
-void USubtitlesSubsystem::RemoveActiveSubtitle(TObjectPtr<const USubtitleAssetUserData> Subtitle)
+void USubtitlesSubsystem::RemoveActiveSubtitle(const FSubtitleAssetData& Subtitle)
 {
 	bool bSuccessfullyRemoved = false;
 	int32 FirstRemovedIndex = Algo::StableRemoveIf(ActiveSubtitles, [&Subtitle](const FActiveSubtitle& ActiveSubtitle) {return ActiveSubtitle.Subtitle == Subtitle; });
@@ -247,11 +242,11 @@ void USubtitlesSubsystem::RemoveActiveSubtitle(TObjectPtr<const USubtitleAssetUs
 	// Stop Displaying the removed subtitle and display a newly-most-relevant one if applicable.
 	if (bSuccessfullyRemoved && IsValid(SubtitleWidget))
 	{
-		SubtitleWidget->StopDisplayingSubtitle(Subtitle->SubtitleType);
+		SubtitleWidget->StopDisplayingSubtitle(Subtitle.SubtitleType);
 
 		if (!ActiveSubtitles.IsEmpty())
 		{
-			const USubtitleAssetUserData& HighestPrioritySubtitle = *ActiveSubtitles[0].Subtitle;
+			const FSubtitleAssetData& HighestPrioritySubtitle = ActiveSubtitles[0].Subtitle;
 			SubtitleWidget->StartDisplayingSubtitle(HighestPrioritySubtitle);
 		}
 	}
@@ -306,7 +301,7 @@ void USubtitlesSubsystem::UpdateWidgetData()
 			bInitializedWidget = true;
 		}
 
-		const USubtitleAssetUserData& NewHighestPrioritySubtitle = *ActiveSubtitles[0].Subtitle;
+		const FSubtitleAssetData& NewHighestPrioritySubtitle = ActiveSubtitles[0].Subtitle;
 		SubtitleWidget->StartDisplayingSubtitle(NewHighestPrioritySubtitle);
 	}
 	else
diff --git a/Engine/Plugins/Experimental/SubtitlesAndClosedCaptions/Source/SubtitlesAndClosedCaptions/Public/ActiveSubtitle.h b/Engine/Plugins/Experimental/SubtitlesAndClosedCaptions/Source/SubtitlesAndClosedCaptions/Public/ActiveSubtitle.h
--- a/Engine/Plugins/Experimental/SubtitlesAndClosedCaptions/Source/SubtitlesAndClosedCaptions/Public/ActiveSubtitle.h
+++ b/Engine/Plugins/Experimental/SubtitlesAndClosedCaptions/Source/SubtitlesAndClosedCaptions/Public/ActiveSubtitle.h
@@ -2,9 +2,10 @@
 
 #pragma once 
 
-class USubtitleAssetUserData;
+struct FSubtitleAssetData;
 
 #include "Engine/TimerHandle.h"
+#include "Subtitles/SubtitlesAndClosedCaptionsDelegates.h"
 #include "UObject/ObjectPtr.h"
 
 #include "ActiveSubtitle.generated.h"
@@ -15,7 +16,7 @@ struct FActiveSubtitle
 	GENERATED_BODY()
 
 	UPROPERTY()
-	TObjectPtr<const USubtitleAssetUserData> Subtitle = nullptr;
+	FSubtitleAssetData Subtitle;
 
 	FTimerHandle DurationTimerHandle;
 };
diff --git a/Engine/Plugins/Experimental/SubtitlesAndClosedCaptions/Source/SubtitlesAndClosedCaptions/Public/MovieSceneSubtitlesSystem.h b/Engine/Plugins/Experimental/SubtitlesAndClosedCaptions/Source/SubtitlesAndClosedCaptions/Public/MovieSceneSubtitlesSystem.h
--- a/Engine/Plugins/Experimental/SubtitlesAndClosedCaptions/Source/SubtitlesAndClosedCaptions/Public/MovieSceneSubtitlesSystem.h
+++ b/Engine/Plugins/Experimental/SubtitlesAndClosedCaptions/Source/SubtitlesAndClosedCaptions/Public/MovieSceneSubtitlesSystem.h
@@ -57,12 +57,11 @@ struct FEvaluateSubtitles
 			return;
 		}
 
-		const USubtitleAssetUserData* Subtitle = SubtitleData.SubtitleSection->GetSubtitle();
-		if (!ensureMsgf(Subtitle, TEXT("No valid subtitle found in subtitle section!")))
+		const USubtitleAssetUserData* SubtitleAsset = SubtitleData.SubtitleSection->GetSubtitle();
+		if (!ensureMsgf(SubtitleAsset, TEXT("No valid subtitle found in subtitle section!")))
 		{
 			return;
 		}
-		const UAssetUserData& AssetUserData = *CastChecked<const UAssetUserData>(Subtitle);
 
 		const EMovieScenePlayerStatus::Type Status = Context.GetStatus();
 #if !NO_LOGGING
@@ -87,25 +86,43 @@ struct FEvaluateSubtitles
 			FSubtitlesAndClosedCaptionsDelegates::StopAllSubtitles.ExecuteIfBound();
 		}
 
-		// Only queue up the subtitle if it isn't already active so we don't spam queueing of subtitles every frame
 		if (FSubtitlesAndClosedCaptionsDelegates::IsSubtitleActive.IsBound())
 		{
 			const FFrameNumber LastSubtitleFrame = SubtitleSection->GetRange().GetUpperBoundValue();
 			const FFrameNumber FirstSubtitleFrame = SubtitleSection->GetRange().GetLowerBoundValue();
 			const FFrameNumber LastEvaluatedFrame = Context.GetFrameNumberRange().GetUpperBoundValue();
 
-			// Queue the subtitle to display in short increments whenever playback is over the subtitle.
-			// This allows for jumping and scrubbing the sequence without the subtitle getting stuck on or off.
 			if (LastEvaluatedFrame >= FirstSubtitleFrame && LastEvaluatedFrame < LastSubtitleFrame)
 			{
-				FQueueSubtitleParameters Params{ AssetUserData };
-				Params.Duration = 0.25f;
-				FSubtitlesAndClosedCaptionsDelegates::QueueSubtitle.ExecuteIfBound(Params, ESubtitleTiming::InternallyTimed);
+				const FFrameRate FrameRate = Context.GetFrameRate();
+				for (const FSubtitleAssetData& Entry : SubtitleAsset->Subtitles)
+				{
+					const FFrameTime StartOffset = Entry.StartOffset * FrameRate;
+					const FFrameNumber FirstOffsetSubtitleFrame = FirstSubtitleFrame + StartOffset.FrameNumber;
+
+					// Only need to test for the initial offset for the beginning; the end of the section includes the delay when it is first initialized.
+					if (LastEvaluatedFrame >= FirstOffsetSubtitleFrame)
+					{
+						// Queue the subtitle to display in short (1/4 second) increments whenever playback is over the subtitle.
+						// This allows for jumping and scrubbing the sequence without the subtitle getting stuck on or off.
+						FQueueSubtitleParameters Params{ Entry };
+						Params.Duration = 0.25f;
+						Params.StartOffset = 0.f;
+
+						FSubtitlesAndClosedCaptionsDelegates::QueueSubtitle.ExecuteIfBound(Params, ESubtitleTiming::InternallyTimed);
+					}
+				}
 			}
-			else if(FSubtitlesAndClosedCaptionsDelegates::IsSubtitleActive.Execute(AssetUserData) && (LastSubtitleFrame <= LastEvaluatedFrame))
+			else if(LastSubtitleFrame <= LastEvaluatedFrame)
 			{
-				// If the subtitle is already active and outside the range of frames it should display for, remove it.
-				FSubtitlesAndClosedCaptionsDelegates::StopSubtitle.ExecuteIfBound(AssetUserData);
+				for (const FSubtitleAssetData& Entry : SubtitleAsset->Subtitles)
+				{
+					if (FSubtitlesAndClosedCaptionsDelegates::IsSubtitleActive.Execute(Entry))
+					{
+						// If the subtitle is already active and outside the range of frames it should display for, remove it.
+						FSubtitlesAndClosedCaptionsDelegates::StopSubtitle.ExecuteIfBound(Entry);
+					}
+				}
 			}
 		}
 	}
diff --git a/Engine/Plugins/Experimental/SubtitlesAndClosedCaptions/Source/SubtitlesAndClosedCaptions/Public/SubtitleTextBlock.h b/Engine/Plugins/Experimental/SubtitlesAndClosedCaptions/Source/SubtitlesAndClosedCaptions/Public/SubtitleTextBlock.h
--- a/Engine/Plugins/Experimental/SubtitlesAndClosedCaptions/Source/SubtitlesAndClosedCaptions/Public/SubtitleTextBlock.h
+++ b/Engine/Plugins/Experimental/SubtitlesAndClosedCaptions/Source/SubtitlesAndClosedCaptions/Public/SubtitleTextBlock.h
@@ -27,6 +27,6 @@ private:
 	TObjectPtr<UTextBlock> DescriptionSubtitleBlock;
 
 public:
-	void StartDisplayingSubtitle(const USubtitleAssetUserData& Subtitle);
+	void StartDisplayingSubtitle(const FSubtitleAssetData& Subtitle);
 	void StopDisplayingSubtitle(const ESubtitleType SubtitleType);
 };
\ No newline at end of file
diff --git a/Engine/Plugins/Experimental/SubtitlesAndClosedCaptions/Source/SubtitlesAndClosedCaptions/Public/SubtitlesBlueprintFunctionLibrary.h b/Engine/Plugins/Experimental/SubtitlesAndClosedCaptions/Source/SubtitlesAndClosedCaptions/Public/SubtitlesBlueprintFunctionLibrary.h
--- a/Engine/Plugins/Experimental/SubtitlesAndClosedCaptions/Source/SubtitlesAndClosedCaptions/Public/SubtitlesBlueprintFunctionLibrary.h
+++ b/Engine/Plugins/Experimental/SubtitlesAndClosedCaptions/Source/SubtitlesAndClosedCaptions/Public/SubtitlesBlueprintFunctionLibrary.h
@@ -8,8 +8,6 @@
 
 #define UE_API SUBTITLESANDCLOSEDCAPTIONS_API
 
-class USubtitleAssetUserData;
-
 UCLASS(MinimalAPI, meta = (ScriptName = "SubtitlesLibrary"))
 class USubtitlesBlueprintFunctionLibrary : public UBlueprintFunctionLibrary
 {
@@ -18,13 +16,13 @@ class USubtitlesBlueprintFunctionLibrary : public UBlueprintFunctionLibrary
 public:
 
 	UFUNCTION(BlueprintCallable, Category = Subtitles)
-	static UE_API void QueueSubtitle(const USubtitleAssetUserData* Subtitle, const ESubtitleTiming Timing = ESubtitleTiming::InternallyTimed);
+	static UE_API void QueueSubtitle(const FSubtitleAssetData& Subtitle, const ESubtitleTiming Timing = ESubtitleTiming::InternallyTimed);
 
 	UFUNCTION(BlueprintCallable, Category = Subtitles)
-	static UE_API bool IsSubtitleActive(const USubtitleAssetUserData* Subtitle);
+	static UE_API bool IsSubtitleActive(const FSubtitleAssetData& Subtitle);
 
 	UFUNCTION(BlueprintCallable, Category = Subtitles)
-	static UE_API void StopSubtitle(const USubtitleAssetUserData* Subtitle);
+	static UE_API void StopSubtitle(const FSubtitleAssetData& Subtitle);
 
 	UFUNCTION(BlueprintCallable, Category = Subtitles)
 	static UE_API void StopAllSubtitles();
diff --git a/Engine/Plugins/Experimental/SubtitlesAndClosedCaptions/Source/SubtitlesAndClosedCaptions/Public/SubtitlesSubsystem.h b/Engine/Plugins/Experimental/SubtitlesAndClosedCaptions/Source/SubtitlesAndClosedCaptions/Public/SubtitlesSubsystem.h
--- a/Engine/Plugins/Experimental/SubtitlesAndClosedCaptions/Source/SubtitlesAndClosedCaptions/Public/SubtitlesSubsystem.h
+++ b/Engine/Plugins/Experimental/SubtitlesAndClosedCaptions/Source/SubtitlesAndClosedCaptions/Public/SubtitlesSubsystem.h
@@ -9,7 +9,6 @@
 
 #include "SubtitleTextBlock.h"
 
-#include "Subtitles/SubtitlesAndClosedCaptionsDelegates.h"
 #include "SubtitlesSubsystem.generated.h"
 
 #define UE_API SUBTITLESANDCLOSEDCAPTIONS_API
@@ -50,10 +49,10 @@ public:
 	UE_API virtual void QueueSubtitle(const FQueueSubtitleParameters& Params, const ESubtitleTiming Timing = ESubtitleTiming::InternallyTimed);
 
 	// Returns true if the given subtitle asset is being displayed.
-	UE_API virtual bool IsSubtitleActive(const UAssetUserData& Data) const;
+	UE_API virtual bool IsSubtitleActive(const FSubtitleAssetData& Data) const;
 
 	// Stops the given subtitle asset being displayed.  This includes subtitles not yet being displayed due to their StartOffset
-	UE_API virtual void StopSubtitle(const UAssetUserData& Data);
+	UE_API virtual void StopSubtitle(const FSubtitleAssetData& Data);
 
 	// Stops all queued subtitles from being displayed.  This includes subtitles not yet being displayed due to their StartOffset.
 	UE_API virtual void StopAllSubtitles();
@@ -75,9 +74,14 @@ protected:
 	TArray<FActiveSubtitle> DelayedSubtitles;
 
 protected:
-	UE_API virtual void AddActiveSubtitle(const USubtitleAssetUserData& Subtitle, float Duration, const ESubtitleTiming Timing);
-	UE_API virtual void MakeDelayedSubtitleActive(TObjectPtr<const USubtitleAssetUserData> Subtitle, const ESubtitleTiming Timing);
-	UE_API virtual void RemoveActiveSubtitle(TObjectPtr<const USubtitleAssetUserData> Subtitle);
+	UE_API virtual void AddActiveSubtitle(const FSubtitleAssetData& Subtitle, float Duration, const float StartOffset, const ESubtitleTiming Timing);
+
+	// These need to be UFUNCTIONs for the Timer Delegate bindings.
+	UFUNCTION()
+	UE_API virtual void MakeDelayedSubtitleActive(const FSubtitleAssetData& Subtitle, const ESubtitleTiming Timing);
+
+	UFUNCTION()
+	UE_API virtual void RemoveActiveSubtitle(const FSubtitleAssetData& Subtitle);
 
 private:
 	UE_API void AddAndDisplaySubtitle(FActiveSubtitle& NewActiveSubtitle);
@@ -93,18 +97,17 @@ public:
 	}
 
 	// As this is for testing, assume that the number of subtitles has already been checked.
-	const USubtitleAssetUserData* GetTopRankedSubtitle() const
+	const FSubtitleAssetData& GetTopRankedSubtitle() const
 	{
 		check(ActiveSubtitles.Num() > 0);
 		const FActiveSubtitle& ActiveSubtitle = ActiveSubtitles[0];
 
-		return ActiveSubtitle.Subtitle.Get();
+		return ActiveSubtitle.Subtitle;
 	}
 
-	void TestActivatingDelayedSubtitle(const UAssetUserData& Data)
+	void TestActivatingDelayedSubtitle(const FSubtitleAssetData& Data)
 	{
-		const USubtitleAssetUserData* Subtitle = CastChecked<const USubtitleAssetUserData>(&Data);
-		MakeDelayedSubtitleActive(Subtitle, ESubtitleTiming::InternallyTimed);
+		MakeDelayedSubtitleActive(Data, ESubtitleTiming::InternallyTimed);
 	}
 #endif
 	//	WITH_DEV_AUTOMATION_TESTS
diff --git a/Engine/Plugins/Experimental/SubtitlesAndClosedCaptions/Source/SubtitlesAndClosedCaptionsEditor/Private/SubtitleSequencerSection.cpp b/Engine/Plugins/Experimental/SubtitlesAndClosedCaptions/Source/SubtitlesAndClosedCaptionsEditor/Private/SubtitleSequencerSection.cpp
--- a/Engine/Plugins/Experimental/SubtitlesAndClosedCaptions/Source/SubtitlesAndClosedCaptionsEditor/Private/SubtitleSequencerSection.cpp
+++ b/Engine/Plugins/Experimental/SubtitlesAndClosedCaptions/Source/SubtitlesAndClosedCaptionsEditor/Private/SubtitleSequencerSection.cpp
@@ -12,9 +12,13 @@ FText FSubtitleSequencerSection::GetSectionTitle() const
 	if (TStrongObjectPtr<const UMovieSceneSection> SectionPin = WeakSection.Pin())
 	{
 		const UMovieSceneSubtitleSection& SubtitlesSection = *CastChecked<const UMovieSceneSubtitleSection>(SectionPin.Get());
-		if (const USubtitleAssetUserData* Subtitle = SubtitlesSection.GetSubtitle())
+		if (const USubtitleAssetUserData* SubtitleAsset = SubtitlesSection.GetSubtitle())
 		{
-			return Subtitle->Text;
+			const TArray<FSubtitleAssetData>& Subtitles = SubtitleAsset->Subtitles;
+			if (!Subtitles.IsEmpty())
+			{
+				return Subtitles[0].Text;
+			}
 		}
 	}
 
diff --git a/Engine/Source/Runtime/Engine/Classes/Sound/DialogueSoundWaveProxy.h b/Engine/Source/Runtime/Engine/Classes/Sound/DialogueSoundWaveProxy.h
--- a/Engine/Source/Runtime/Engine/Classes/Sound/DialogueSoundWaveProxy.h
+++ b/Engine/Source/Runtime/Engine/Classes/Sound/DialogueSoundWaveProxy.h
@@ -52,4 +52,9 @@ public:
 	virtual USoundClass* GetSoundClass() const override;
 
 	virtual class UCurveTable* GetCurveData() const override;
+
+	inline const FText& GetSubtitleText() const
+	{
+		return SubtitleText;
+	}
 };
diff --git a/Engine/Source/Runtime/Engine/Private/DialogueWave.cpp b/Engine/Source/Runtime/Engine/Private/DialogueWave.cpp
--- a/Engine/Source/Runtime/Engine/Private/DialogueWave.cpp
+++ b/Engine/Source/Runtime/Engine/Private/DialogueWave.cpp
@@ -486,26 +486,7 @@ void UDialogueSoundWaveProxy::Parse(class FAudioDevice* AudioDevice, const UPTRI
 		// Add in the subtitle if they exist
 		if (ActiveSound.bHandleSubtitles && bHasSubtitles)
 		{
-			if (UseNewSubtitles)
-			{
-				IModularFeatures& ModularFeatures = IModularFeatures::Get();
-				if (ModularFeatures.IsModularFeatureAvailable(ISubtitlesAndClosedCaptionsModule::GetModularFeatureName()))
-				{
-					USubtitleAssetUserData* SubtitleCue = NewObject<USubtitleAssetUserData>(GetTransientPackage(), NAME_None, RF_Transient);
-					SubtitleCue->Text = SubtitleText;
-					SubtitleCue->Priority = ActiveSound.SubtitlePriority;
-					SubtitleCue->Duration = GetDuration();
-					SubtitleCue->StartOffset = ActiveSound.RequestedStartTime;
-
-					const ISubtitlesAndClosedCaptionsModule& SubtitlesModule = ModularFeatures.GetModularFeature<ISubtitlesAndClosedCaptionsModule>(ISubtitlesAndClosedCaptionsModule::GetModularFeatureName());
-					FSubtitlesAndClosedCaptionsDelegates::QueueSubtitle.ExecuteIfBound(FQueueSubtitleParameters{ *SubtitleCue }, ESubtitleTiming::InternallyTimed);
-				}
-				else
-				{
-					UE_LOG(LogAudio, Warning, TEXT("The SubtitlesAndClosedCaptions plugin needs to be enabled for subtitles to be displayed."), *GetPathName());
-				}
-			}
-			else
+			if (!UseNewSubtitles)
 			{
 				FQueueSubtitleParams QueueSubtitleParams(Subtitles);
 				{
diff --git a/Engine/Source/Runtime/Engine/Private/SoundWave.cpp b/Engine/Source/Runtime/Engine/Private/SoundWave.cpp
--- a/Engine/Source/Runtime/Engine/Private/SoundWave.cpp
+++ b/Engine/Source/Runtime/Engine/Private/SoundWave.cpp
@@ -3212,11 +3212,15 @@ void USoundWave::PostEditChangeProperty(FPropertyChangedEvent& PropertyChangedEv
 			if (IsValid(LastModifiedUserData))
 			{
 				USubtitleAssetUserData* UserDataAsSubtitle = Cast<USubtitleAssetUserData>(LastModifiedUserData.Get() );
-				if (IsValid(UserDataAsSubtitle) && UserDataAsSubtitle->Duration == SubtitleDefaultDuration && UserDataAsSubtitle->Text.IsEmpty())
+				if (IsValid(UserDataAsSubtitle) && !UserDataAsSubtitle->Subtitles.IsEmpty())
 				{
-					// If the sound is looping indefinitely, the subtitle duration will default to the length of the first loop.
-					// It is still editable by the user after this default is set.
-					UserDataAsSubtitle->Duration = Duration;
+					FSubtitleAssetData& Subtitle = UserDataAsSubtitle->Subtitles.Last();
+					if(Subtitle.Duration == SubtitleDefaultDuration && Subtitle.Text.IsEmpty())
+					{
+						// If the sound is looping indefinitely, the subtitle duration will default to the length of the first loop.
+						// It is still editable by the user after this default is set.
+						Subtitle.Duration = Duration;
+					}
 				}
 			}
 		}
@@ -5056,7 +5060,7 @@ void USoundWave::SerializeCuePoints(FArchive& Ar, const bool bIsLoadingFromCooke
 
 		float ResampleRatio = -1.0f;
 		const ITargetPlatform* CookingTarget = Ar.CookingTarget();
-		if (ensure(CookingTarget))
+		if (ensure(CookingTarget) && CookingTarget != nullptr)
 		{
 			if (const FPlatformAudioCookOverrides* Overrides = FPlatformCompressionUtilities::GetCookOverrides(*CookingTarget->IniPlatformName()))
 			{
diff --git a/Engine/Source/Runtime/Engine/Private/Subtitles/SubtitlesAndClosedCaptionsDelegates.cpp b/Engine/Source/Runtime/Engine/Private/Subtitles/SubtitlesAndClosedCaptionsDelegates.cpp
--- a/Engine/Source/Runtime/Engine/Private/Subtitles/SubtitlesAndClosedCaptionsDelegates.cpp
+++ b/Engine/Source/Runtime/Engine/Private/Subtitles/SubtitlesAndClosedCaptionsDelegates.cpp
@@ -5,6 +5,6 @@
 #include UE_INLINE_GENERATED_CPP_BY_NAME(SubtitlesAndClosedCaptionsDelegates)
 
 TDelegate<void(const FQueueSubtitleParameters&, const ESubtitleTiming)> FSubtitlesAndClosedCaptionsDelegates::QueueSubtitle;
-TDelegate<bool(const UAssetUserData&)> FSubtitlesAndClosedCaptionsDelegates::IsSubtitleActive;
-TDelegate<void(const UAssetUserData&)> FSubtitlesAndClosedCaptionsDelegates::StopSubtitle;
+TDelegate<bool(const FSubtitleAssetData&)> FSubtitlesAndClosedCaptionsDelegates::IsSubtitleActive;
+TDelegate<void(const FSubtitleAssetData&)> FSubtitlesAndClosedCaptionsDelegates::StopSubtitle;
 TDelegate<void()> FSubtitlesAndClosedCaptionsDelegates::StopAllSubtitles;
diff --git a/Engine/Source/Runtime/Engine/Public/Subtitles/SubtitlesAndClosedCaptionsDelegates.h b/Engine/Source/Runtime/Engine/Public/Subtitles/SubtitlesAndClosedCaptionsDelegates.h
--- a/Engine/Source/Runtime/Engine/Public/Subtitles/SubtitlesAndClosedCaptionsDelegates.h
+++ b/Engine/Source/Runtime/Engine/Public/Subtitles/SubtitlesAndClosedCaptionsDelegates.h
@@ -43,16 +43,11 @@ static constexpr float SubtitleMinDuration		= 0.05f;
 // Default value to initialize subtitle duration to. Used by SoundWaves to check if they should manually set the duration.
 static constexpr float SubtitleDefaultDuration	= 3.f;
 
-/**
- *  Base class for subtitle data being attached to assets
- */
-UCLASS(MinimalAPI, BlueprintType)
-class USubtitleAssetUserData : public UAssetUserData
+USTRUCT(MinimalAPI, BlueprintType)
+struct FSubtitleAssetData
 {
 	GENERATED_BODY()
 
-public:
-
 	// #SUBTITLES_PRD: carried over from FSubtitleCue, still required
 
 	// The text to appear in the subtitle.
@@ -90,16 +85,51 @@ public:
 
 #if WITH_EDITORONLY_DATA
 	// Author comments, editor-only.
-	UPROPERTY(EditAnywhere, BlueprintReadOnly, Category = "Subtitles")
+	UPROPERTY(EditAnywhere, Category = "Subtitles")
 	FString Comment;
 #endif // WITH_EDITORONLY_DATA
+
+	// The struct needs a custom equality test operator because FText doesn't have one overloaded: need to use EqualTo().
+	// Don't compare the Comment property because it isn't compiled outside the editor.
+	const bool operator==(const FSubtitleAssetData& Other) const
+	{
+		return ((Text.EqualTo(Other.Text) && (SubtitleDurationType == Other.SubtitleDurationType) && (Duration == Other.Duration) 
+			&& (StartOffset == Other.StartOffset) && (Priority == Other.Priority) && (SubtitleType == Other.SubtitleType)));
+	}
+};
+
+/**
+ *  Base class for subtitle data being attached to assets
+ */
+UCLASS(MinimalAPI, BlueprintType)
+class USubtitleAssetUserData : public UAssetUserData
+{
+	GENERATED_BODY()
+
+public:
+	UPROPERTY(EditAnywhere, BlueprintReadOnly, Category = "Subtitles")
+	TArray<FSubtitleAssetData> Subtitles;
+
+	// In seconds, the longest any subtitle in the Subtitles array will last, including any initial offset starts.
+	float GetMaximumDuration() const
+	{
+		float Duration = 0.f;
+		for (const FSubtitleAssetData& Entry : Subtitles)
+		{
+			const float EntryDuration = Entry.Duration + Entry.StartOffset;
+			Duration = FMath::Max(Duration, EntryDuration);
+		}
+
+		return Duration;
+	};
 };
 
 
 struct FQueueSubtitleParameters
 {
-	const UAssetUserData& Subtitle;
+	const FSubtitleAssetData& Subtitle;
 	TOptional<float> Duration;
+	TOptional<float> StartOffset;	// Scrubbing subtitles in the sequencer needs an override for StartOffset.
 };
 
 class FSubtitlesAndClosedCaptionsDelegates
@@ -109,9 +139,9 @@ public:
 	// Have the subtitle subsystem to queue a subtitle to be displayed
 	static UE_API TDelegate<void(const FQueueSubtitleParameters&, const ESubtitleTiming)> QueueSubtitle;
 
-	static UE_API TDelegate<bool(const UAssetUserData&)> IsSubtitleActive;
+	static UE_API TDelegate<bool(const FSubtitleAssetData&)> IsSubtitleActive;
 
-	static UE_API TDelegate<void(const UAssetUserData&)> StopSubtitle;
+	static UE_API TDelegate<void(const FSubtitleAssetData&)> StopSubtitle;
 
 	static UE_API TDelegate<void()> StopAllSubtitles;
 };
