Add bulk subtitle import from overlays and SRT/subriptext files. Copying to Release 38 branch.

#virtualized

diff --git a/Engine/Plugins/Experimental/SubtitlesAndClosedCaptions/Source/SubtitlesAndClosedCaptionsEditor/Private/SubtitleAssetActionUtility.cpp b/Engine/Plugins/Experimental/SubtitlesAndClosedCaptions/Source/SubtitlesAndClosedCaptionsEditor/Private/SubtitleAssetActionUtility.cpp
new file mode 100644
--- /dev/null
+++ b/Engine/Plugins/Experimental/SubtitlesAndClosedCaptions/Source/SubtitlesAndClosedCaptionsEditor/Private/SubtitleAssetActionUtility.cpp
@@ -0,0 +1,56 @@
+// Copyright Epic Games, Inc. All Rights Reserved.
+
+#include "SubtitleAssetActionUtility.h"
+#include "AssetRegistry/AssetRegistryModule.h"
+#include "AssetToolsModule.h"
+#include "BasicOverlays.h"
+#include "EditorUtilityLibrary.h"
+#include "IAssetTools.h"
+#include "Subtitles/SubtitlesAndClosedCaptionsDelegates.h"
+
+USubtitleAssetActionUtility::USubtitleAssetActionUtility()
+{
+	SupportedClasses.Add(UBasicOverlays::StaticClass());
+}
+
+// .SRT (SubRip Text) files are imported by the engine as UBasicOverlays, which contain an Array
+// of captions with timestamps and positions. This utility creates a new asset that converts BasicOverlays to Subtitle assets.
+void USubtitleAssetActionUtility::ConvertBasicOverlaysToSubtitles()
+{
+	TArray<UObject*> SelectedAssets = UEditorUtilityLibrary::GetSelectedAssets();
+
+	for (UObject* Asset : SelectedAssets)
+	{
+		UBasicOverlays* OverlayAsset = Cast<UBasicOverlays>(Asset);
+		if (IsValid(OverlayAsset))
+		{
+			const FString NameSuffix = TEXT("Subtitle");
+			const FString AssetPath = FPaths::ProjectContentDir();
+			FString AssetName = OverlayAsset->GetName() + NameSuffix;
+			FString PackagePath = OverlayAsset->GetPackage()->GetPathName() + NameSuffix;
+
+			UPackage* Package = CreatePackage(*PackagePath);
+			USubtitleAssetUserData* NewSubtitleAsset = NewObject<USubtitleAssetUserData>(
+				Package, *AssetName, EObjectFlags::RF_Public | EObjectFlags::RF_Standalone
+			);
+
+			// Iterate through each timecode in the overlay and create a new subtitle from it.
+			for (const FOverlayItem& Overlay : OverlayAsset->GetAllOverlays())
+			{
+				FSubtitleAssetData Subtitle;
+				Subtitle.Text = FText::FromString(Overlay.Text);
+				
+				// Timing
+				Subtitle.StartOffset = Overlay.StartTime.GetTotalSeconds();
+				Subtitle.Duration = Overlay.EndTime.GetTotalSeconds() - Overlay.StartTime.GetTotalSeconds();
+				Subtitle.SubtitleDurationType = ESubtitleDurationType::UseDurationProperty;
+				
+				// Default priority, subtitle type, etc are fine.
+				NewSubtitleAsset->Subtitles.Add(Subtitle);
+			}
+
+			FAssetRegistryModule::AssetCreated(NewSubtitleAsset);
+			NewSubtitleAsset->MarkPackageDirty();
+		}
+	}
+}
\ No newline at end of file
diff --git a/Engine/Plugins/Experimental/SubtitlesAndClosedCaptions/Source/SubtitlesAndClosedCaptionsEditor/Public/SubtitleAssetActionUtility.h b/Engine/Plugins/Experimental/SubtitlesAndClosedCaptions/Source/SubtitlesAndClosedCaptionsEditor/Public/SubtitleAssetActionUtility.h
new file mode 100644
--- /dev/null
+++ b/Engine/Plugins/Experimental/SubtitlesAndClosedCaptions/Source/SubtitlesAndClosedCaptionsEditor/Public/SubtitleAssetActionUtility.h
@@ -0,0 +1,22 @@
+// Copyright Epic Games, Inc. All Rights Reserved.
+
+#pragma once
+#include "CoreMinimal.h"
+#include "AssetActionUtility.h"
+
+#include "SubtitleAssetActionUtility.generated.h"
+
+// Asset Action Utility class for any editor right-click actions relevant to subtitles.
+UCLASS()
+class USubtitleAssetActionUtility : public UAssetActionUtility
+{
+	GENERATED_BODY()
+
+public:
+	USubtitleAssetActionUtility();
+
+	UFUNCTION(CallInEditor)
+	void ConvertBasicOverlaysToSubtitles();
+
+
+};
\ No newline at end of file
diff --git a/Engine/Plugins/Experimental/SubtitlesAndClosedCaptions/Source/SubtitlesAndClosedCaptionsEditor/SubtitlesAndClosedCaptionsEditor.Build.cs b/Engine/Plugins/Experimental/SubtitlesAndClosedCaptions/Source/SubtitlesAndClosedCaptionsEditor/SubtitlesAndClosedCaptionsEditor.Build.cs
--- a/Engine/Plugins/Experimental/SubtitlesAndClosedCaptions/Source/SubtitlesAndClosedCaptionsEditor/SubtitlesAndClosedCaptionsEditor.Build.cs
+++ b/Engine/Plugins/Experimental/SubtitlesAndClosedCaptions/Source/SubtitlesAndClosedCaptionsEditor/SubtitlesAndClosedCaptionsEditor.Build.cs
@@ -13,8 +13,10 @@ namespace UnrealBuildTool.Rules
 				new string[]
 				{
 					"AssetTools",
+					"Blutility",
 					"LevelSequence",
 					"MovieScene",
+					"Overlay",
 					"Sequencer",
 					"SequencerCore",
 					"Slate",
