Add automatic subtitle playback support for DialogueWaves inside SoundCues. Merging to Release-38.00

diff --git a/Engine/Plugins/Experimental/SubtitlesAndClosedCaptions/Source/SubtitlesAndClosedCaptions/Private/SubtitlesAudioSubsystem.cpp b/Engine/Plugins/Experimental/SubtitlesAndClosedCaptions/Source/SubtitlesAndClosedCaptions/Private/SubtitlesAudioSubsystem.cpp
--- a/Engine/Plugins/Experimental/SubtitlesAndClosedCaptions/Source/SubtitlesAndClosedCaptions/Private/SubtitlesAudioSubsystem.cpp
+++ b/Engine/Plugins/Experimental/SubtitlesAndClosedCaptions/Source/SubtitlesAndClosedCaptions/Private/SubtitlesAudioSubsystem.cpp
@@ -4,6 +4,8 @@
 
 #include "ActiveSound.h"
 #include "Sound/DialogueSoundWaveProxy.h"
+#include "Sound/SoundCue.h"
+#include "Sound/SoundNodeDialoguePlayer.h"
 #include "Subtitles/SubtitlesAndClosedCaptionsDelegates.h"
 
 #include UE_INLINE_GENERATED_CPP_BY_NAME(SubtitlesAudioSubsystem)
@@ -27,16 +29,32 @@ void USubtitlesAudioSubsystem::NotifyActiveSoundCreated(FActiveSound& ActiveSoun
 		// Also check for DialogueWaves with localized text
 		if (ActiveSound.bHandleSubtitles)
 		{
-			const UDialogueSoundWaveProxy* DialogueWave = Cast<UDialogueSoundWaveProxy>(ActiveSound.GetSound());
+			FSubtitleAssetData SubtitleCue;
+
+			const UDialogueSoundWaveProxy* DialogueWave = Cast<UDialogueSoundWaveProxy>(Sound);
 			if (IsValid(DialogueWave))
 			{
-				FSubtitleAssetData SubtitleCue;
-				SubtitleCue.Text = DialogueWave->GetSubtitleText();
-				SubtitleCue.Priority = ActiveSound.SubtitlePriority;
-				SubtitleCue.Duration = DialogueWave->GetDuration();
-				SubtitleCue.StartOffset = ActiveSound.RequestedStartTime;
+				QueueSubtitleFromDialogueWaveProxy(DialogueWave, ActiveSound);
+			}
+			else
+			{
+				// Check inside Sound Cues for DialogueWaves as well.
+				const USoundCue* SoundCue = Cast<USoundCue>(Sound);
+				if (IsValid(SoundCue))
+				{
+					const bool bQueuedFromTopLevelNode = CheckSoundCueNodeForDialogueWaves(SoundCue->FirstNode, ActiveSound);
 
-				FSubtitlesAndClosedCaptionsDelegates::QueueSubtitle.ExecuteIfBound(FQueueSubtitleParameters{ SubtitleCue }, ESubtitleTiming::InternallyTimed);
+					// If the parent node is not a DialogueWavePlayer, check for nested nodes that may also have DialogueWaves
+					if(!bQueuedFromTopLevelNode)
+					{
+						TArray<USoundNode*> ChildNodes;
+						SoundCue->FirstNode->GetAllNodes(ChildNodes);
+						for (const USoundNode* Node : ChildNodes)
+						{
+							CheckSoundCueNodeForDialogueWaves(Node, ActiveSound);
+						}
+					}
+				}
 			}
 		}
 	}
@@ -65,3 +83,30 @@ void USubtitlesAudioSubsystem::NotifyActiveSoundDeleting(const FActiveSound& Act
 		}
 	}
 }
+
+bool USubtitlesAudioSubsystem::CheckSoundCueNodeForDialogueWaves(const USoundNode* Node, const FActiveSound& ActiveSound)
+{
+	const USoundNodeDialoguePlayer* DialoguePlayerNode = Cast<USoundNodeDialoguePlayer>(Node);
+	if (IsValid(DialoguePlayerNode))
+	{
+		const UDialogueSoundWaveProxy* Proxy = DialoguePlayerNode->GetContextualDialogueSoundWaveProxy();
+		check(Proxy);
+
+		QueueSubtitleFromDialogueWaveProxy(Proxy, ActiveSound);
+		return true;
+	}
+
+	return false;
+}
+
+void USubtitlesAudioSubsystem::QueueSubtitleFromDialogueWaveProxy(const UDialogueSoundWaveProxy* DialogueWaveProxy, const FActiveSound& ActiveSound)
+{
+	check(ActiveSound.bHandleSubtitles);	// Function should only be called if the Active Sound actually has subtitles it needs queued.
+	FSubtitleAssetData SubtitleCue;
+	SubtitleCue.Text = DialogueWaveProxy->GetSubtitleText();
+	SubtitleCue.Priority = ActiveSound.SubtitlePriority;
+	SubtitleCue.Duration = DialogueWaveProxy->GetDuration();
+	SubtitleCue.StartOffset = ActiveSound.RequestedStartTime;
+
+	FSubtitlesAndClosedCaptionsDelegates::QueueSubtitle.ExecuteIfBound(FQueueSubtitleParameters{ SubtitleCue }, ESubtitleTiming::InternallyTimed);
+}
diff --git a/Engine/Plugins/Experimental/SubtitlesAndClosedCaptions/Source/SubtitlesAndClosedCaptions/Public/SubtitlesAudioSubsystem.h b/Engine/Plugins/Experimental/SubtitlesAndClosedCaptions/Source/SubtitlesAndClosedCaptions/Public/SubtitlesAudioSubsystem.h
--- a/Engine/Plugins/Experimental/SubtitlesAndClosedCaptions/Source/SubtitlesAndClosedCaptions/Public/SubtitlesAudioSubsystem.h
+++ b/Engine/Plugins/Experimental/SubtitlesAndClosedCaptions/Source/SubtitlesAndClosedCaptions/Public/SubtitlesAudioSubsystem.h
@@ -7,6 +7,9 @@
 
 #define UE_API SUBTITLESANDCLOSEDCAPTIONS_API
 
+class UDialogueSoundWaveProxy;
+class USoundNode;
+
 /*
 *	SubtitlesAudioSubsystem - AudioEngineSubsystem for automatically queueing subtitles on audio classes.
 *
@@ -25,6 +28,11 @@ public:
 	virtual void NotifyActiveSoundCreated(FActiveSound& ActiveSound) override;
 	virtual void NotifyActiveSoundDeleting(const FActiveSound& ActiveSound) override;
 	//~ End IActiveSoundUpdateInterface
+
+private:
+	// Helper functions to cut down on duplicated code
+	static bool CheckSoundCueNodeForDialogueWaves(const USoundNode* Node, const FActiveSound& ActiveSound);
+	static void QueueSubtitleFromDialogueWaveProxy(const UDialogueSoundWaveProxy* DialogueWaveProxy, const FActiveSound& ActiveSound);
 };
 
 #undef UE_API
diff --git a/Engine/Source/Runtime/Engine/Classes/Sound/SoundNodeDialoguePlayer.h b/Engine/Source/Runtime/Engine/Classes/Sound/SoundNodeDialoguePlayer.h
--- a/Engine/Source/Runtime/Engine/Classes/Sound/SoundNodeDialoguePlayer.h
+++ b/Engine/Source/Runtime/Engine/Classes/Sound/SoundNodeDialoguePlayer.h
@@ -8,6 +8,7 @@
 #include "SoundNodeDialoguePlayer.generated.h"
 
 class FAudioDevice;
+class UDialogueSoundWaveProxy;
 class UDialogueWave;
 struct FActiveSound;
 struct FSoundParseParameters;
@@ -33,6 +34,9 @@ public:
 	ENGINE_API UDialogueWave* GetDialogueWave() const;
 	ENGINE_API void SetDialogueWave(UDialogueWave* Value);
 
+	// Subtitles need to be able to grab the UDialogueSoundWaveProxy to for localized text and duration.
+	ENGINE_API const UDialogueSoundWaveProxy* GetContextualDialogueSoundWaveProxy() const;
+
 	//~ Begin USoundNode Interface
 	virtual int32 GetMaxChildNodes() const override;
 	virtual float GetDuration();
diff --git a/Engine/Source/Runtime/Engine/Private/SoundNodeDialoguePlayer.cpp b/Engine/Source/Runtime/Engine/Private/SoundNodeDialoguePlayer.cpp
--- a/Engine/Source/Runtime/Engine/Private/SoundNodeDialoguePlayer.cpp
+++ b/Engine/Source/Runtime/Engine/Private/SoundNodeDialoguePlayer.cpp
@@ -3,6 +3,7 @@
 #include "Sound/SoundNodeDialoguePlayer.h"
 #include "ActiveSound.h"
 #include "Sound/DialogueWave.h"
+#include "Sound/DialogueSoundWaveProxy.h"
 
 #include UE_INLINE_GENERATED_CPP_BY_NAME(SoundNodeDialoguePlayer)
 
@@ -91,6 +92,22 @@ void USoundNodeDialoguePlayer::SetDialogueWave(UDialogueWave* Value)
 	DialogueWaveParameter.DialogueWave = Value;
 }
 
+const UDialogueSoundWaveProxy* USoundNodeDialoguePlayer::GetContextualDialogueSoundWaveProxy() const
+{
+	const UDialogueWave* DialogueWave = GetDialogueWave();
+	check(IsValid(DialogueWave));
+
+	const USoundBase* SoundBase = DialogueWave ? DialogueWave->GetWaveFromContext(DialogueWaveParameter.Context) : NULL;
+	if (IsValid(SoundBase))
+	{
+		const UDialogueSoundWaveProxy* Proxy = Cast<UDialogueSoundWaveProxy>(SoundBase);
+		check(IsValid(Proxy));	// GetWaveFromContext() specifically uses DialogueSoundWaveProxy, but casts to SoundBase before returning.
+		return Proxy;
+	}
+
+	return nullptr;
+}
+
 UDialogueWave* USoundNodeDialoguePlayer::GetDialogueWave() const
 {
 	return DialogueWaveParameter.DialogueWave;
