Add an optional Border UI element to subtitle widget TextBlocks Merging Main to Release-38.00
#virtualized

diff --git a/Engine/Plugins/Experimental/SubtitlesAndClosedCaptions/Source/SubtitlesAndClosedCaptions/Private/SubtitleTextBlock.cpp b/Engine/Plugins/Experimental/SubtitlesAndClosedCaptions/Source/SubtitlesAndClosedCaptions/Private/SubtitleTextBlock.cpp
--- a/Engine/Plugins/Experimental/SubtitlesAndClosedCaptions/Source/SubtitlesAndClosedCaptions/Private/SubtitleTextBlock.cpp
+++ b/Engine/Plugins/Experimental/SubtitlesAndClosedCaptions/Source/SubtitlesAndClosedCaptions/Private/SubtitleTextBlock.cpp
@@ -1,8 +1,9 @@
 // Copyright Epic Games, Inc. All Rights Reserved.
 
 #include "SubtitleTextBlock.h"
-#include "SubtitlesAndClosedCaptionsModule.h"
+#include "Components/Border.h"
 #include "Components/TextBlock.h"
+#include "SubtitlesAndClosedCaptionsModule.h"
 
 #include UE_INLINE_GENERATED_CPP_BY_NAME(SubtitleTextBlock)
 
@@ -25,12 +26,27 @@ void USubtitleTextBlock::NativeConstruct()
 	{
 		DescriptionSubtitleBlock->SetVisibility(ESlateVisibility::Collapsed);
 	}
+
+	// Borders should also not be visible on startup.
+	if (IsValid(DialogBorder))
+	{
+		DialogBorder->SetVisibility(ESlateVisibility::Collapsed);
+	}
+	if (IsValid(CaptionBorder))
+	{
+		CaptionBorder->SetVisibility(ESlateVisibility::Collapsed);
+	}
+	if (IsValid(DescriptionBorder))
+	{
+		DescriptionBorder->SetVisibility(ESlateVisibility::Collapsed);
+	}
 }
 
 void USubtitleTextBlock::StartDisplayingSubtitle(const FSubtitleAssetData& Subtitle)
 {
 	// Pick which UextBlock is relevant by category....
 	UTextBlock* TextToModify = nullptr;
+	UBorder* BorderToModify = nullptr;
 
 	switch (Subtitle.SubtitleType)
 	{
@@ -43,18 +59,21 @@ void USubtitleTextBlock::StartDisplayingSubtitle(const FSubtitleAssetData& Subti
 	case ESubtitleType::Subtitle:
 	{
 		TextToModify = DialogSubtitleBlock;
+		BorderToModify = DialogBorder;
 		break;
 	}
 
 	case ESubtitleType::ClosedCaption:
 	{
 		TextToModify = CaptionSubtitleBlock;
+		BorderToModify = CaptionBorder;
 		break;
 	}
 
 	case ESubtitleType::AudioDescription:
 	{
 		TextToModify = DescriptionSubtitleBlock;
+		BorderToModify = DescriptionBorder;
 		break;
 	}
 
@@ -65,39 +84,53 @@ void USubtitleTextBlock::StartDisplayingSubtitle(const FSubtitleAssetData& Subti
 	{
 		TextToModify->SetText(Subtitle.Text);
 		TextToModify->SetVisibility(ESlateVisibility::HitTestInvisible);
+
+		// Only hide/display border if the corresponding text is also valid.
+		if (IsValid(BorderToModify))
+		{
+			BorderToModify->SetVisibility(ESlateVisibility::HitTestInvisible);
+		}
 	}
 }
 
 void USubtitleTextBlock::StopDisplayingSubtitle(const ESubtitleType SubtitleType)
 {
-
 	UTextBlock* TextToModify = nullptr;
+	UBorder* BorderToModify = nullptr;
+
 	switch (SubtitleType)
 	{
 	default: // Fallthrough, no break (see ::StartDisplayingSubtitle)
 	case ESubtitleType::Subtitle:
 	{
 		TextToModify = DialogSubtitleBlock;
+		BorderToModify = DialogBorder;
 		break;
 	}
 
 	case ESubtitleType::ClosedCaption:
 	{
 		TextToModify = CaptionSubtitleBlock;
+		BorderToModify = CaptionBorder;
 		break;
 	}
 
 	case ESubtitleType::AudioDescription:
 	{
 		TextToModify = DescriptionSubtitleBlock;
+		BorderToModify = DescriptionBorder;
 		break;
 	}
 
 	}
 
-
 	if (IsValid(TextToModify))
 	{
 		TextToModify->SetVisibility(ESlateVisibility::Collapsed);
+
+		if (IsValid(BorderToModify))
+		{
+			BorderToModify->SetVisibility(ESlateVisibility::Collapsed);
+		}
 	}
 }
diff --git a/Engine/Plugins/Experimental/SubtitlesAndClosedCaptions/Source/SubtitlesAndClosedCaptions/Public/SubtitleTextBlock.h b/Engine/Plugins/Experimental/SubtitlesAndClosedCaptions/Source/SubtitlesAndClosedCaptions/Public/SubtitleTextBlock.h
--- a/Engine/Plugins/Experimental/SubtitlesAndClosedCaptions/Source/SubtitlesAndClosedCaptions/Public/SubtitleTextBlock.h
+++ b/Engine/Plugins/Experimental/SubtitlesAndClosedCaptions/Source/SubtitlesAndClosedCaptions/Public/SubtitleTextBlock.h
@@ -6,6 +6,7 @@
 
 #include "SubtitleTextBlock.generated.h"
 
+class UBorder;
 class UTextBlock;
 
 UCLASS(Abstract, BlueprintType, Blueprintable, meta=( DontUseGenericSpawnObject="True", DisableNativeTick), MinimalAPI)
@@ -17,6 +18,7 @@ protected:
 	virtual void NativeConstruct() override;
 
 private:
+	// Subtitle Text Blocks
 	UPROPERTY(EditDefaultsOnly, Category = "Subtitles", meta = (BindWidget))
 	TObjectPtr<UTextBlock> DialogSubtitleBlock;
 
@@ -26,6 +28,15 @@ private:
 	UPROPERTY(EditDefaultsOnly, Category = "Subtitles", meta = (BindWidget))
 	TObjectPtr<UTextBlock> DescriptionSubtitleBlock;
 
+	// Borders to the text blocks
+	UPROPERTY(EditDefaultsOnly, Category = "SubtitleBorders", meta = (BindWidget))
+	TObjectPtr<UBorder> DialogBorder;
+
+	UPROPERTY(EditDefaultsOnly, Category = "SubtitleBorders", meta = (BindWidget))
+	TObjectPtr<UBorder> CaptionBorder;
+
+	UPROPERTY(EditDefaultsOnly, Category = "SubtitleBorders", meta = (BindWidget))
+	TObjectPtr<UBorder> DescriptionBorder;
 public:
 	void StartDisplayingSubtitle(const FSubtitleAssetData& Subtitle);
 	void StopDisplayingSubtitle(const ESubtitleType SubtitleType);
