Use BindWidgetOptional,BlueprintReadOnly on SubtitleTextBlocks, Merging Main to Release-38.00

diff --git a/Engine/Plugins/Experimental/SubtitlesAndClosedCaptions/Source/SubtitlesAndClosedCaptions/Public/SubtitleTextBlock.h b/Engine/Plugins/Experimental/SubtitlesAndClosedCaptions/Source/SubtitlesAndClosedCaptions/Public/SubtitleTextBlock.h
--- a/Engine/Plugins/Experimental/SubtitlesAndClosedCaptions/Source/SubtitlesAndClosedCaptions/Public/SubtitleTextBlock.h
+++ b/Engine/Plugins/Experimental/SubtitlesAndClosedCaptions/Source/SubtitlesAndClosedCaptions/Public/SubtitleTextBlock.h
@@ -17,25 +17,24 @@ class USubtitleTextBlock : public UUserWidget
 protected:
 	virtual void NativeConstruct() override;
 
-private:
 	// Subtitle Text Blocks
-	UPROPERTY(EditDefaultsOnly, Category = "Subtitles", meta = (BindWidget))
+	UPROPERTY(BlueprintReadOnly, Category = "Subtitles", meta = (BindWidgetOptional))
 	TObjectPtr<UTextBlock> DialogSubtitleBlock;
 
-	UPROPERTY(EditDefaultsOnly, Category = "Subtitles", meta = (BindWidget))
+	UPROPERTY(BlueprintReadOnly, Category = "Subtitles", meta = (BindWidgetOptional))
 	TObjectPtr<UTextBlock> CaptionSubtitleBlock;
 
-	UPROPERTY(EditDefaultsOnly, Category = "Subtitles", meta = (BindWidget))
+	UPROPERTY(BlueprintReadOnly, Category = "Subtitles", meta = (BindWidgetOptional))
 	TObjectPtr<UTextBlock> DescriptionSubtitleBlock;
 
 	// Borders to the text blocks
-	UPROPERTY(EditDefaultsOnly, Category = "SubtitleBorders", meta = (BindWidget))
+	UPROPERTY(BlueprintReadOnly, Category = "SubtitleBorders", meta = (BindWidgetOptional))
 	TObjectPtr<UBorder> DialogBorder;
 
-	UPROPERTY(EditDefaultsOnly, Category = "SubtitleBorders", meta = (BindWidget))
+	UPROPERTY(BlueprintReadOnly, Category = "SubtitleBorders", meta = (BindWidgetOptional))
 	TObjectPtr<UBorder> CaptionBorder;
 
-	UPROPERTY(EditDefaultsOnly, Category = "SubtitleBorders", meta = (BindWidget))
+	UPROPERTY(BlueprintReadOnly, Category = "SubtitleBorders", meta = (BindWidgetOptional))
 	TObjectPtr<UBorder> DescriptionBorder;
 public:
 	void StartDisplayingSubtitle(const FSubtitleAssetData& Subtitle);
