Migrate Subtitle Migration Actions to native code. Add undo/redo support. Migrate legacy comments correctly.
Merging Main to Release-38.00
#virtualized

diff --git a/Engine/Plugins/Experimental/SubtitlesAndClosedCaptions/Source/SubtitlesAndClosedCaptionsEditor/Private/SubtitleMigrationUtility.cpp b/Engine/Plugins/Experimental/SubtitlesAndClosedCaptions/Source/SubtitlesAndClosedCaptionsEditor/Private/SubtitleMigrationUtility.cpp
new file mode 100644
--- /dev/null
+++ b/Engine/Plugins/Experimental/SubtitlesAndClosedCaptions/Source/SubtitlesAndClosedCaptionsEditor/Private/SubtitleMigrationUtility.cpp
@@ -0,0 +1,158 @@
+// Copyright Epic Games, Inc. All Rights Reserved.
+
+#include "SubtitleMigrationUtility.h"
+
+#include "EditorUtilityLibrary.h"
+#include "Engine/Engine.h"
+#include "Sound/SoundWave.h"
+#include "SubtitlesAndClosedCaptionsModule.h"
+#include "Subtitles/SubtitlesAndClosedCaptionsDelegates.h"
+
+const FString USubtitleMigrationUtility::UndoContext = TEXT("SubtitleMigration");
+
+USubtitleMigrationUtility::USubtitleMigrationUtility()
+{
+	SupportedClasses.Add(USoundBase::StaticClass());
+}
+
+void USubtitleMigrationUtility::AddBlankSubtitle()
+{
+#if WITH_EDITOR
+	if (GEngine)
+	{
+		// Setup transaction for undoing, and describe said action.
+		const static FText Description = FText::FromString(TEXT("Add Blank Subtitle"));
+		GEngine->BeginTransaction(*UndoContext, Description, nullptr);
+
+		const TArray<UObject*> SelectedAssets = UEditorUtilityLibrary::GetSelectedAssetsOfClass(USoundBase::StaticClass());
+		for (UObject* SelectedObject : SelectedAssets)
+		{
+			if (!IsValid(SelectedObject))
+			{
+				continue;
+			}
+
+			// GetSelectedAssetsOfClass and SupportedClasses should already only provide USoundBase
+			USoundBase* SelectedSound = CastChecked<USoundBase>(SelectedObject);
+
+			// Add an asset userdata if one doesn't already exist.
+			if (!SelectedSound->HasAssetUserDataOfClass(USubtitleAssetUserData::StaticClass()))
+			{
+				SelectedSound->AddAssetUserDataOfClass(USubtitleAssetUserData::StaticClass());
+			}
+
+			UAssetUserData* UserData = SelectedSound->GetAssetUserDataOfClass(USubtitleAssetUserData::StaticClass());
+			USubtitleAssetUserData* SubtitleData = CastChecked<USubtitleAssetUserData>(UserData);
+
+			// Mark the userdata as modified for the undo action to work correctly.
+			SubtitleData->Modify();
+
+			SubtitleData->Subtitles.Add(FSubtitleAssetData());
+		}
+
+		// End transaction for undoing.
+		GEngine->EndTransaction();
+	}
+#endif
+}
+
+void USubtitleMigrationUtility::RemoveLegacySubtitles()
+{
+#if WITH_EDITOR
+	if (GEngine)
+	{
+		// Setup transaction for undoing, and describe said action.
+		const static FText Description = FText::FromString(TEXT("Remove legacy subtitles from SoundWaves."));
+		GEngine->BeginTransaction(*UndoContext, Description, nullptr);
+
+		const TArray<UObject*> SelectedAssets = UEditorUtilityLibrary::GetSelectedAssetsOfClass(USoundWave::StaticClass());
+		for (UObject* SelectedObject : SelectedAssets)
+		{
+			if (!IsValid(SelectedObject))
+			{
+				continue;
+			}
+
+			// This tool supports other non-wave Soundbases for other actions, but only USoundWaves have subtitles to migrate in this fashion.
+			USoundWave* SelectedSound = Cast<USoundWave>(SelectedObject);
+
+			// Only empty out the legacy subtitles if there's already new subtitles in place.
+			if (IsValid(SelectedSound))
+			{
+				if (SelectedSound->HasAssetUserDataOfClass(USubtitleAssetUserData::StaticClass()))
+				{
+					SelectedSound->Modify();
+					SelectedSound->Subtitles.Empty();
+				}
+				else
+				{
+					UE_LOG(LogSubtitlesAndClosedCaptions, Warning, TEXT("Can't remove legacy subtitles: migrate them to the new system first with Convert Legacy Subtitles."));
+				}
+			}
+			else
+			{
+				UE_LOG(LogSubtitlesAndClosedCaptions, Warning, TEXT("Remove Legacy Subtitles can only be used on SoundWave assets."));
+			}
+		}
+
+		GEngine->EndTransaction();
+	}
+#endif
+}
+
+void USubtitleMigrationUtility::ConvertLegacySubtitles()
+{
+#if WITH_EDITOR
+	if (GEngine)
+	{
+		// Setup transaction for undoing, and describe said action.
+		const static FText Description = FText::FromString(TEXT("Make Asset UserData from Legacy Subtitles"));
+		GEngine->BeginTransaction(*UndoContext, Description, nullptr);
+
+		const TArray<UObject*> SelectedAssets = UEditorUtilityLibrary::GetSelectedAssetsOfClass(USoundWave::StaticClass());
+		for (UObject* SelectedObject : SelectedAssets)
+		{
+			if (!IsValid(SelectedObject))
+			{
+				continue;
+			}
+
+			// This tool supports other non-wave Soundbases for other actions, but only USoundWaves have subtitles to migrate in this fashion.
+			USoundWave* SelectedSound = Cast<USoundWave>(SelectedObject);
+
+			if (IsValid(SelectedSound))
+			{
+				// Add an asset userdata if one doesn't already exist.
+				if (!SelectedSound->HasAssetUserDataOfClass(USubtitleAssetUserData::StaticClass()))
+				{
+					SelectedSound->AddAssetUserDataOfClass(USubtitleAssetUserData::StaticClass());
+				}
+
+				UAssetUserData* UserData = SelectedSound->GetAssetUserDataOfClass(USubtitleAssetUserData::StaticClass());
+				USubtitleAssetUserData* SubtitleData = CastChecked<USubtitleAssetUserData>(UserData);
+
+				// Loop through all the legacy subtitles and add a new subtitle to the UserData for each.
+				for (const FSubtitleCue& Cue : SelectedSound->Subtitles)
+				{
+					FSubtitleAssetData Data;
+					Data.SubtitleType = ESubtitleType::Subtitle;
+					Data.SubtitleDurationType = ESubtitleDurationType::UseSoundDuration;
+
+					Data.Priority = SelectedSound->SubtitlePriority;
+					Data.Text = Cue.Text;
+					Data.StartOffset = Cue.Time;
+					Data.Comment = SelectedSound->Comment; // While it looks like this is the sound's comment, the property is in the Subtitles category.
+					SubtitleData->Modify();
+					SubtitleData->Subtitles.Add(Data);
+				}
+			}
+			else
+			{
+				UE_LOG(LogSubtitlesAndClosedCaptions, Warning, TEXT("Convert Legacy Subtitles can only be used on SoundWave assets."));
+			}
+		}
+
+		GEngine->EndTransaction();
+	}
+#endif
+}
\ No newline at end of file
diff --git a/Engine/Plugins/Experimental/SubtitlesAndClosedCaptions/Source/SubtitlesAndClosedCaptionsEditor/Public/SubtitleMigrationUtility.h b/Engine/Plugins/Experimental/SubtitlesAndClosedCaptions/Source/SubtitlesAndClosedCaptionsEditor/Public/SubtitleMigrationUtility.h
new file mode 100644
--- /dev/null
+++ b/Engine/Plugins/Experimental/SubtitlesAndClosedCaptions/Source/SubtitlesAndClosedCaptionsEditor/Public/SubtitleMigrationUtility.h
@@ -0,0 +1,30 @@
+// Copyright Epic Games, Inc. All Rights Reserved.
+
+#pragma once
+#include "CoreMinimal.h"
+#include "AssetActionUtility.h"
+
+#include "SubtitleMigrationUtility.generated.h"
+
+// Asset Action Utility class for migrating subtitles
+// Associated asset SubtitleMigrationAction.uasset required to show up in editor.
+UCLASS()
+class USubtitleMigrationUtility : public UAssetActionUtility
+{
+	GENERATED_BODY()
+
+private:
+	const static FString UndoContext;
+
+public:
+	USubtitleMigrationUtility();
+
+	UFUNCTION(CallInEditor)
+	void AddBlankSubtitle();
+
+	UFUNCTION(CallInEditor)
+	void RemoveLegacySubtitles();
+
+	UFUNCTION(CallInEditor)
+	void ConvertLegacySubtitles();
+};
