Change USubtitleTextBlock class name to USubtitleWidget. Merging Main to Release-38.00

diff --git a/Engine/Plugins/Experimental/SubtitlesAndClosedCaptions/Source/SubtitlesAndClosedCaptions/Private/SubtitleTextBlock.cpp b/Engine/Plugins/Experimental/SubtitlesAndClosedCaptions/Source/SubtitlesAndClosedCaptions/Private/SubtitleTextBlock.cpp
deleted file mode 100644
diff --git a/Engine/Plugins/Experimental/SubtitlesAndClosedCaptions/Source/SubtitlesAndClosedCaptions/Private/SubtitleWidget.cpp b/Engine/Plugins/Experimental/SubtitlesAndClosedCaptions/Source/SubtitlesAndClosedCaptions/Private/SubtitleWidget.cpp
new file mode 100644
--- /dev/null
+++ b/Engine/Plugins/Experimental/SubtitlesAndClosedCaptions/Source/SubtitlesAndClosedCaptions/Private/SubtitleWidget.cpp
@@ -0,0 +1,136 @@
+// Copyright Epic Games, Inc. All Rights Reserved.
+
+#include "SubtitleWidget.h"
+#include "Components/Border.h"
+#include "Components/TextBlock.h"
+#include "SubtitlesAndClosedCaptionsModule.h"
+
+#include UE_INLINE_GENERATED_CPP_BY_NAME(SubtitleWidget)
+
+void USubtitleWidget::NativeConstruct()
+{
+	Super::NativeConstruct();
+
+	// Ensure the text blocks don't start visible.
+	if (IsValid(DialogSubtitleBlock))
+	{
+		DialogSubtitleBlock->SetVisibility(ESlateVisibility::Collapsed);
+	}
+
+	if (IsValid(CaptionSubtitleBlock))
+	{
+		CaptionSubtitleBlock->SetVisibility(ESlateVisibility::Collapsed);
+	}
+
+	if (IsValid(DescriptionSubtitleBlock))
+	{
+		DescriptionSubtitleBlock->SetVisibility(ESlateVisibility::Collapsed);
+	}
+
+	// Borders should also not be visible on startup.
+	if (IsValid(DialogBorder))
+	{
+		DialogBorder->SetVisibility(ESlateVisibility::Collapsed);
+	}
+	if (IsValid(CaptionBorder))
+	{
+		CaptionBorder->SetVisibility(ESlateVisibility::Collapsed);
+	}
+	if (IsValid(DescriptionBorder))
+	{
+		DescriptionBorder->SetVisibility(ESlateVisibility::Collapsed);
+	}
+}
+
+void USubtitleWidget::StartDisplayingSubtitle(const FSubtitleAssetData& Subtitle)
+{
+	// Pick which UextBlock is relevant by category....
+	UTextBlock* TextToModify = nullptr;
+	UBorder* BorderToModify = nullptr;
+
+	switch (Subtitle.SubtitleType)
+	{
+	default:
+	{
+		UE_LOG(LogSubtitlesAndClosedCaptions, Warning, TEXT("An unrecognized subtitle type was queued for display. Using the standard Subtitle text block as a fallback."));
+		// Fallthrough, no break.
+	}
+
+	case ESubtitleType::Subtitle:
+	{
+		TextToModify = DialogSubtitleBlock;
+		BorderToModify = DialogBorder;
+		break;
+	}
+
+	case ESubtitleType::ClosedCaption:
+	{
+		TextToModify = CaptionSubtitleBlock;
+		BorderToModify = CaptionBorder;
+		break;
+	}
+
+	case ESubtitleType::AudioDescription:
+	{
+		TextToModify = DescriptionSubtitleBlock;
+		BorderToModify = DescriptionBorder;
+		break;
+	}
+
+	}
+
+	// ...Then modify and display it
+	if (IsValid(TextToModify))
+	{
+		TextToModify->SetText(Subtitle.Text);
+		TextToModify->SetVisibility(ESlateVisibility::HitTestInvisible);
+
+		// Only hide/display border if the corresponding text is also valid.
+		if (IsValid(BorderToModify))
+		{
+			BorderToModify->SetVisibility(ESlateVisibility::HitTestInvisible);
+		}
+	}
+}
+
+void USubtitleWidget::StopDisplayingSubtitle(const ESubtitleType SubtitleType)
+{
+	UTextBlock* TextToModify = nullptr;
+	UBorder* BorderToModify = nullptr;
+
+	switch (SubtitleType)
+	{
+	default: // Fallthrough, no break (see ::StartDisplayingSubtitle)
+	case ESubtitleType::Subtitle:
+	{
+		TextToModify = DialogSubtitleBlock;
+		BorderToModify = DialogBorder;
+		break;
+	}
+
+	case ESubtitleType::ClosedCaption:
+	{
+		TextToModify = CaptionSubtitleBlock;
+		BorderToModify = CaptionBorder;
+		break;
+	}
+
+	case ESubtitleType::AudioDescription:
+	{
+		TextToModify = DescriptionSubtitleBlock;
+		BorderToModify = DescriptionBorder;
+		break;
+	}
+
+	}
+
+	if (IsValid(TextToModify))
+	{
+		TextToModify->SetVisibility(ESlateVisibility::Collapsed);
+
+		if (IsValid(BorderToModify))
+		{
+			BorderToModify->SetVisibility(ESlateVisibility::Collapsed);
+		}
+	}
+}
diff --git a/Engine/Plugins/Experimental/SubtitlesAndClosedCaptions/Source/SubtitlesAndClosedCaptions/Private/SubtitlesSubsystem.cpp b/Engine/Plugins/Experimental/SubtitlesAndClosedCaptions/Source/SubtitlesAndClosedCaptions/Private/SubtitlesSubsystem.cpp
--- a/Engine/Plugins/Experimental/SubtitlesAndClosedCaptions/Source/SubtitlesAndClosedCaptions/Private/SubtitlesSubsystem.cpp
+++ b/Engine/Plugins/Experimental/SubtitlesAndClosedCaptions/Source/SubtitlesAndClosedCaptions/Private/SubtitlesSubsystem.cpp
@@ -270,18 +270,18 @@ bool USubtitlesSubsystem::TryCreateUMGWidget()
 	// Set up the UMG widget
 	const USubtitlesSettings* Settings = GetDefault<USubtitlesSettings>();
 	check(Settings != nullptr);
-	const TSubclassOf<USubtitleTextBlock>& WidgetToUse = Settings->GetWidget();
+	const TSubclassOf<USubtitleWidget>& WidgetToUse = Settings->GetWidget();
 	if (IsValid(WidgetToUse))
 	{
-		SubtitleWidget = CreateWidget<USubtitleTextBlock>(GetWorld(), WidgetToUse);
+		SubtitleWidget = CreateWidget<USubtitleWidget>(GetWorld(), WidgetToUse);
 	}
 	else
 	{
 		// Fallback to default widget (not set by user):
-		const TSubclassOf<USubtitleTextBlock>& WidgetToUseDefault = Settings->GetWidgetDefault();
+		const TSubclassOf<USubtitleWidget>& WidgetToUseDefault = Settings->GetWidgetDefault();
 		if (IsValid(WidgetToUseDefault))
 		{
-			SubtitleWidget = CreateWidget<USubtitleTextBlock>(GetWorld(), WidgetToUseDefault);
+			SubtitleWidget = CreateWidget<USubtitleWidget>(GetWorld(), WidgetToUseDefault);
 		}
 	}
 	bInitializedWidget = false;
diff --git a/Engine/Plugins/Experimental/SubtitlesAndClosedCaptions/Source/SubtitlesAndClosedCaptions/Public/SubtitleTextBlock.h b/Engine/Plugins/Experimental/SubtitlesAndClosedCaptions/Source/SubtitlesAndClosedCaptions/Public/SubtitleTextBlock.h
deleted file mode 100644
diff --git a/Engine/Plugins/Experimental/SubtitlesAndClosedCaptions/Source/SubtitlesAndClosedCaptions/Public/SubtitleWidget.h b/Engine/Plugins/Experimental/SubtitlesAndClosedCaptions/Source/SubtitlesAndClosedCaptions/Public/SubtitleWidget.h
new file mode 100644
--- /dev/null
+++ b/Engine/Plugins/Experimental/SubtitlesAndClosedCaptions/Source/SubtitlesAndClosedCaptions/Public/SubtitleWidget.h
@@ -0,0 +1,42 @@
+// Copyright Epic Games, Inc. All Rights Reserved.
+
+#pragma once
+#include "Blueprint/UserWidget.h"
+#include "Subtitles/SubtitlesAndClosedCaptionsDelegates.h"
+
+#include "SubtitleWidget.generated.h"
+
+class UBorder;
+class UTextBlock;
+
+UCLASS(Abstract, BlueprintType, Blueprintable, meta=( DontUseGenericSpawnObject="True", DisableNativeTick), MinimalAPI)
+class USubtitleWidget : public UUserWidget
+{
+	GENERATED_BODY()
+
+protected:
+	virtual void NativeConstruct() override;
+
+	// Subtitle Text Blocks
+	UPROPERTY(BlueprintReadOnly, Category = "Subtitles", meta = (BindWidgetOptional))
+	TObjectPtr<UTextBlock> DialogSubtitleBlock;
+
+	UPROPERTY(BlueprintReadOnly, Category = "Subtitles", meta = (BindWidgetOptional))
+	TObjectPtr<UTextBlock> CaptionSubtitleBlock;
+
+	UPROPERTY(BlueprintReadOnly, Category = "Subtitles", meta = (BindWidgetOptional))
+	TObjectPtr<UTextBlock> DescriptionSubtitleBlock;
+
+	// Borders to the text blocks
+	UPROPERTY(BlueprintReadOnly, Category = "SubtitleBorders", meta = (BindWidgetOptional))
+	TObjectPtr<UBorder> DialogBorder;
+
+	UPROPERTY(BlueprintReadOnly, Category = "SubtitleBorders", meta = (BindWidgetOptional))
+	TObjectPtr<UBorder> CaptionBorder;
+
+	UPROPERTY(BlueprintReadOnly, Category = "SubtitleBorders", meta = (BindWidgetOptional))
+	TObjectPtr<UBorder> DescriptionBorder;
+public:
+	void StartDisplayingSubtitle(const FSubtitleAssetData& Subtitle);
+	void StopDisplayingSubtitle(const ESubtitleType SubtitleType);
+};
\ No newline at end of file
diff --git a/Engine/Plugins/Experimental/SubtitlesAndClosedCaptions/Source/SubtitlesAndClosedCaptions/Public/SubtitlesSettings.h b/Engine/Plugins/Experimental/SubtitlesAndClosedCaptions/Source/SubtitlesAndClosedCaptions/Public/SubtitlesSettings.h
--- a/Engine/Plugins/Experimental/SubtitlesAndClosedCaptions/Source/SubtitlesAndClosedCaptions/Public/SubtitlesSettings.h
+++ b/Engine/Plugins/Experimental/SubtitlesAndClosedCaptions/Source/SubtitlesAndClosedCaptions/Public/SubtitlesSettings.h
@@ -4,7 +4,7 @@
 
 #include "Engine/DeveloperSettings.h"
 #include "Math/Color.h"
-#include "SubtitleTextBlock.h"
+#include "SubtitleWidget.h"
 #include "SubtitlesSettings.generated.h"
 
 
@@ -16,14 +16,14 @@ class USubtitlesSettings : public UDeveloperSettings
 public:
 	USubtitlesSettings();
 
-	const TSubclassOf<USubtitleTextBlock>& GetWidget() const { return SubtitleWidgetToUse; }
-	const TSubclassOf<USubtitleTextBlock>& GetWidgetDefault() const { return SubtitleWidgetToUseDefault; }
+	const TSubclassOf<USubtitleWidget>& GetWidget() const { return SubtitleWidgetToUse; }
+	const TSubclassOf<USubtitleWidget>& GetWidgetDefault() const { return SubtitleWidgetToUseDefault; }
 
 protected:
 
 	UPROPERTY(config, EditDefaultsOnly, Category = Subtitles, AdvancedDisplay)
-	TSubclassOf<USubtitleTextBlock> SubtitleWidgetToUse;
+	TSubclassOf<USubtitleWidget> SubtitleWidgetToUse;
 
 	UPROPERTY(config)
-	TSubclassOf<USubtitleTextBlock> SubtitleWidgetToUseDefault; // fallback for SubtitleWidgetToUse (not set by user)
+	TSubclassOf<USubtitleWidget> SubtitleWidgetToUseDefault; // fallback for SubtitleWidgetToUse (not set by user)
 };
diff --git a/Engine/Plugins/Experimental/SubtitlesAndClosedCaptions/Source/SubtitlesAndClosedCaptions/Public/SubtitlesSubsystem.h b/Engine/Plugins/Experimental/SubtitlesAndClosedCaptions/Source/SubtitlesAndClosedCaptions/Public/SubtitlesSubsystem.h
--- a/Engine/Plugins/Experimental/SubtitlesAndClosedCaptions/Source/SubtitlesAndClosedCaptions/Public/SubtitlesSubsystem.h
+++ b/Engine/Plugins/Experimental/SubtitlesAndClosedCaptions/Source/SubtitlesAndClosedCaptions/Public/SubtitlesSubsystem.h
@@ -7,7 +7,7 @@
 #include "Math/Color.h"
 #include "Subsystems/WorldSubsystem.h"
 
-#include "SubtitleTextBlock.h"
+#include "SubtitleWidget.h"
 
 #include "SubtitlesSubsystem.generated.h"
 
@@ -114,7 +114,7 @@ public:
 
 private:
 	UPROPERTY()
-	TObjectPtr<USubtitleTextBlock> SubtitleWidget = nullptr;
+	TObjectPtr<USubtitleWidget> SubtitleWidget = nullptr;
 
 	bool bInitializedWidget = false;
 	UE_API bool TryCreateUMGWidget();
