Subtitle Sequencer: Add support for Drag 'n Drop.

diff --git a/Engine/Plugins/Experimental/SubtitlesAndClosedCaptions/Source/SubtitlesAndClosedCaptionsEditor/Private/SubtitlesTrackEditor.cpp b/Engine/Plugins/Experimental/SubtitlesAndClosedCaptions/Source/SubtitlesAndClosedCaptionsEditor/Private/SubtitlesTrackEditor.cpp
--- a/Engine/Plugins/Experimental/SubtitlesAndClosedCaptions/Source/SubtitlesAndClosedCaptionsEditor/Private/SubtitlesTrackEditor.cpp
+++ b/Engine/Plugins/Experimental/SubtitlesAndClosedCaptions/Source/SubtitlesAndClosedCaptionsEditor/Private/SubtitlesTrackEditor.cpp
@@ -5,10 +5,12 @@
 #include "AssetRegistry/AssetRegistryModule.h"
 #include "ContentBrowserDelegates.h"
 #include "ContentBrowserModule.h"
+#include "DragAndDrop/AssetDragDropOp.h"
 #include "IContentBrowserSingleton.h"
 #include "Input/Reply.h"
 #include "LevelSequence.h"
 #include "MovieSceneSubtitlesTrack.h"
+#include "MovieSceneToolHelpers.h"
 #include "MVVM/Views/ViewUtilities.h"
 #include "SequencerSettings.h"
 #include "Subtitles/SubtitlesAndClosedCaptionsDelegates.h"
@@ -73,6 +75,32 @@ void FSubtitlesTrackEditor::HandleAddMenuEntryExecute()
 	}
 }
 
+bool FSubtitlesTrackEditor::IsDragDropValid(const FDragDropEvent& DragDropEvent, const FSequencerDragDropParams& DragDropParams, DragDropValidInfo& OutDragDropInfo)
+{
+	if (!DragDropParams.Track.IsValid() || !DragDropParams.Track.Get()->IsA(UMovieSceneSubtitlesTrack::StaticClass()))
+	{
+		return false;
+	}
+	OutDragDropInfo.Operation = DragDropEvent.GetOperation();
+	if (!OutDragDropInfo.Operation.IsValid() || !OutDragDropInfo.Operation->IsOfType<FAssetDragDropOp>())
+	{
+		return false;
+	}
+	OutDragDropInfo.SequencerPtr = GetSequencer();
+	if (!OutDragDropInfo.SequencerPtr)
+	{
+		return false;
+	}
+	OutDragDropInfo.FocusedSequence = OutDragDropInfo.SequencerPtr->GetFocusedMovieSceneSequence();
+	if (!OutDragDropInfo.FocusedSequence)
+	{
+		return false;
+	}
+	OutDragDropInfo.DragDropOp = StaticCastSharedPtr<FAssetDragDropOp>(OutDragDropInfo.Operation);
+
+	return true;
+}
+
 TSharedPtr<SWidget> FSubtitlesTrackEditor::BuildOutlinerEditWidget(const FGuid& ObjectBinding, UMovieSceneTrack* Track, const FBuildEditWidgetParams& Params)
 {
 	TDelegate OnAssetSelected = FOnAssetSelected::CreateRaw(this, &FSubtitlesTrackEditor::OnAssetSelected, Track);
@@ -83,6 +111,35 @@ TSharedPtr<SWidget> FSubtitlesTrackEditor::BuildOutlinerEditWidget(const FGuid&
 		, Params.ViewModel);
 }
 
+bool FSubtitlesTrackEditor::HandleAssetAdded(UObject* Asset, const FGuid& TargetObjectGuid)
+{
+	if (Asset->IsA<USubtitleAssetUserData>())
+	{
+		auto Subtitle = Cast<USubtitleAssetUserData>(Asset);
+		UMovieSceneSubtitlesTrack* DummyTrack = nullptr;
+		const FScopedTransaction Transaction(LOCTEXT("AddSubtitle_Transaction", "Add Subtitle"));
+		if (TargetObjectGuid.IsValid())
+		{
+			TArray<TWeakObjectPtr<>> OutObjects;
+			for (TWeakObjectPtr<> Object : GetSequencer()->FindObjectsInCurrentSequence(TargetObjectGuid))
+			{
+				if (Object.IsValid())
+				{
+					OutObjects.Add(Object);
+				}
+			}
+			AnimatablePropertyChanged(FOnKeyProperty::CreateRaw(this, &FSubtitlesTrackEditor::AddNewAttachedSubtitle, Subtitle, DummyTrack, OutObjects));
+		}
+		else
+		{
+			int32 RowIndex = INDEX_NONE;
+			AnimatablePropertyChanged(FOnKeyProperty::CreateRaw(this, &FSubtitlesTrackEditor::AddNewSubtitle, Subtitle, DummyTrack, RowIndex));
+		}
+		return true;
+	}
+	return false;
+}
+
 TSharedRef<SWidget> FSubtitlesTrackEditor::BuildSubMenu(FOnAssetSelected OnAssetSelected, FOnAssetEnterPressed OnAssetEnterPressed)
 {
 	TSharedPtr<ISequencer> SequencerPtr = GetSequencer();
@@ -195,11 +252,171 @@ bool FSubtitlesTrackEditor::SupportsSequence(UMovieSceneSequence* InSequence) co
 
 bool FSubtitlesTrackEditor::OnAllowDrop(const FDragDropEvent& DragDropEvent, FSequencerDragDropParams& DragDropParams)
 {
+	DragDropValidInfo DragDropInfo;
+	if (!IsDragDropValid(DragDropEvent, DragDropParams, DragDropInfo))
+	{
+		return false;
+	}
+	for (const FAssetData& AssetData : DragDropInfo.DragDropOp->GetAssets())
+	{
+		if (!MovieSceneToolHelpers::IsValidAsset(DragDropInfo.FocusedSequence, AssetData))
+		{
+			continue;
+		}
+		if (USubtitleAssetUserData* SubtitleAsset = Cast<USubtitleAssetUserData>(AssetData.GetAsset()))
+		{
+			FFrameRate TickResolution = DragDropInfo.SequencerPtr->GetFocusedTickResolution();
+			FFrameNumber LengthInFrames = TickResolution.AsFrameNumber(SubtitleAsset->GetMaximumDuration());
+			DragDropParams.FrameRange = TRange<FFrameNumber>(DragDropParams.FrameNumber, DragDropParams.FrameNumber + LengthInFrames);
+			return true;
+		}
+	}
 	return false;
 }
 
 FReply FSubtitlesTrackEditor::OnDrop(const FDragDropEvent& DragDropEvent, const FSequencerDragDropParams& DragDropParams)
 {
-	return FReply::Unhandled();
+	DragDropValidInfo DragDropInfo;
+	if (!IsDragDropValid(DragDropEvent, DragDropParams, DragDropInfo))
+	{
+		return FReply::Unhandled();
+	}
+
+	UMovieSceneSubtitlesTrack* SubtitlesTrack = Cast<UMovieSceneSubtitlesTrack>(DragDropParams.Track);
+
+	const FScopedTransaction Transaction(LOCTEXT("DropAssets", "Drop Assets"));
+
+	FMovieSceneTrackEditor::BeginKeying(DragDropParams.FrameNumber);
+
+	bool bAnyDropped = false;
+	for (const FAssetData& AssetData : DragDropInfo.DragDropOp->GetAssets())
+	{
+		if (!MovieSceneToolHelpers::IsValidAsset(DragDropInfo.FocusedSequence, AssetData))
+		{
+			continue;
+		}
+
+		USubtitleAssetUserData* SubtitleAsset = Cast<USubtitleAssetUserData>(AssetData.GetAsset());
+
+		if (SubtitleAsset)
+		{
+			if (DragDropParams.TargetObjectGuid.IsValid())
+			{
+				TArray<TWeakObjectPtr<>> OutObjects;
+				for (TWeakObjectPtr<> Object : DragDropInfo.SequencerPtr->FindObjectsInCurrentSequence(DragDropParams.TargetObjectGuid))
+				{
+					OutObjects.Add(Object);
+				}
+
+				AnimatablePropertyChanged(FOnKeyProperty::CreateRaw(this, &FSubtitlesTrackEditor::AddNewAttachedSubtitle, SubtitleAsset, SubtitlesTrack, OutObjects));
+			}
+			else
+			{
+				AnimatablePropertyChanged(FOnKeyProperty::CreateRaw(this, &FSubtitlesTrackEditor::AddNewSubtitle, SubtitleAsset, SubtitlesTrack, DragDropParams.RowIndex));
+			}
+
+			bAnyDropped = true;
+		}
+	}
+
+	FMovieSceneTrackEditor::EndKeying();
+
+	return bAnyDropped ? FReply::Handled() : FReply::Unhandled();
+}
+
+/** Delegate for AnimatablePropertyChanged in HandleAssetAdded and OnDrop for adding subtitles.  */
+FKeyPropertyResult FSubtitlesTrackEditor::AddNewSubtitle(FFrameNumber KeyTime, USubtitleAssetUserData* Subtitle, UMovieSceneSubtitlesTrack* InDestinationTrack, int32 RowIndex)
+{
+	FKeyPropertyResult KeyPropertyResult;
+
+	UMovieScene* FocusedMovieScene = GetFocusedMovieScene();
+	if (FocusedMovieScene->IsReadOnly())
+	{
+		return KeyPropertyResult;
+	}
+
+	// Undo tracking
+	FocusedMovieScene->Modify();
+
+	FFindOrCreateRootTrackResult<UMovieSceneSubtitlesTrack> TrackResult;
+	TrackResult.Track = InDestinationTrack;
+
+	// Make a new track if a valid one isn't provided.
+	// ::HandleAssetAdded only passes a nullptr, but ::OnDrop should already have a valid destination track (if not, it's safe to create a new one).
+	if (!InDestinationTrack)
+	{
+		TrackResult = FindOrCreateRootTrack<UMovieSceneSubtitlesTrack>();
+		InDestinationTrack = TrackResult.Track;
+	}
+
+	if (ensure(InDestinationTrack))
+	{
+		InDestinationTrack->Modify();
+
+		UMovieSceneSection* NewSection = InDestinationTrack->AddNewSubtitleOnRow(*Subtitle, KeyTime, RowIndex);
+
+		// If a new track was created earlier, set display names and initialize it now.
+		if (TrackResult.bWasCreated)
+		{
+			InDestinationTrack->SetDisplayName(LOCTEXT("SubtitleTrackName", "Subtitle"));
+
+			if (GetSequencer().IsValid())
+			{
+				GetSequencer()->OnAddTrack(InDestinationTrack, FGuid());
+			}
+		}
+
+		KeyPropertyResult.bTrackModified = true;
+		KeyPropertyResult.SectionsCreated.Add(NewSection);
+	}
+
+	return KeyPropertyResult;
+}
+
+/** Delegate for AnimatablePropertyChanged in HandleAssetAdded and OnDrop for attached subtitles.  */
+FKeyPropertyResult FSubtitlesTrackEditor::AddNewAttachedSubtitle(FFrameNumber KeyTime, class USubtitleAssetUserData* Subtitle, UMovieSceneSubtitlesTrack* InDestinationTrack, TArray<TWeakObjectPtr<UObject>> ObjectsToAttachTo)
+{
+	FKeyPropertyResult KeyPropertyResult;
+
+	for (int32 ObjectIndex = 0; ObjectIndex < ObjectsToAttachTo.Num(); ++ObjectIndex)
+	{
+		UObject* Object = ObjectsToAttachTo[ObjectIndex].Get();
+
+		FFindOrCreateHandleResult HandleResult = FindOrCreateHandleToObject(Object);
+		FGuid ObjectHandle = HandleResult.Handle;
+		KeyPropertyResult.bHandleCreated |= HandleResult.bWasCreated;
+
+		if (ObjectHandle.IsValid())
+		{
+			FFindOrCreateTrackResult TrackResult;
+			TrackResult.Track = InDestinationTrack;
+
+			// Make a new track if a valid one isn't provided.
+			// ::HandleAssetAdded only passes a nullptr, but ::OnDrop should already have a valid destination track (if not, it's safe to create a new one).
+			if (!InDestinationTrack)
+			{
+				TrackResult = FindOrCreateTrackForObject(ObjectHandle, UMovieSceneSubtitlesTrack::StaticClass());
+				InDestinationTrack = Cast<UMovieSceneSubtitlesTrack>(TrackResult.Track);
+			}
+
+			KeyPropertyResult.bTrackCreated |= TrackResult.bWasCreated;
+
+			if (ensure(InDestinationTrack))
+			{
+				InDestinationTrack->Modify();
+
+				UMovieSceneSection* NewSection = InDestinationTrack->AddNewSubtitle(*Subtitle, KeyTime);
+				InDestinationTrack->SetDisplayName(LOCTEXT("SubtitleTrackName", "Subtitle"));
+				KeyPropertyResult.bTrackModified = true;
+				KeyPropertyResult.SectionsCreated.Add(NewSection);
+
+				GetSequencer()->EmptySelection();
+				GetSequencer()->SelectSection(NewSection);
+				GetSequencer()->ThrobSectionSelection();
+			}
+		}
+	}
+
+	return KeyPropertyResult;
 }
 #undef LOCTEXT_NAMESPACE
diff --git a/Engine/Plugins/Experimental/SubtitlesAndClosedCaptions/Source/SubtitlesAndClosedCaptionsEditor/Private/SubtitlesTrackEditor.h b/Engine/Plugins/Experimental/SubtitlesAndClosedCaptions/Source/SubtitlesAndClosedCaptionsEditor/Private/SubtitlesTrackEditor.h
--- a/Engine/Plugins/Experimental/SubtitlesAndClosedCaptions/Source/SubtitlesAndClosedCaptionsEditor/Private/SubtitlesTrackEditor.h
+++ b/Engine/Plugins/Experimental/SubtitlesAndClosedCaptions/Source/SubtitlesAndClosedCaptionsEditor/Private/SubtitlesTrackEditor.h
@@ -7,9 +7,19 @@
 #include "ContentBrowserDelegates.h"
 
 struct FGuid;
+class FAssetDragDropOp;
 class FReply;
 class UMovieSceneTrack;
 class UMovieSceneSubtitlesTrack;
+class USubtitleAssetUserData;
+
+struct DragDropValidInfo
+{
+	TSharedPtr<FDragDropOperation> Operation = nullptr;
+	TSharedPtr<ISequencer> SequencerPtr = nullptr;
+	UMovieSceneSequence* FocusedSequence = nullptr;
+	TSharedPtr<FAssetDragDropOp> DragDropOp = nullptr;
+};
 
 /**
  * Tools for subtitles tracks
@@ -40,16 +50,27 @@ public:
 	virtual FText GetDisplayName() const override;
 	virtual void BuildAddTrackMenu(FMenuBuilder& MenuBuilder) override;
 	virtual TSharedPtr<SWidget> BuildOutlinerEditWidget(const FGuid& ObjectBinding, UMovieSceneTrack* Track, const FBuildEditWidgetParams& Params) override;
+	virtual bool HandleAssetAdded(UObject* Asset, const FGuid& TargetObjectGuid) override;
 	virtual TSharedRef<ISequencerSection> MakeSectionInterface(UMovieSceneSection& SectionObject, UMovieSceneTrack& Track, FGuid ObjectBinding) override;
 	virtual bool SupportsType(TSubclassOf<UMovieSceneTrack> Type) const override;
 	virtual bool SupportsSequence(UMovieSceneSequence* InSequence) const override;
 	virtual bool OnAllowDrop(const FDragDropEvent& DragDropEvent, FSequencerDragDropParams& DragDropParams) override;
 	virtual FReply OnDrop(const FDragDropEvent& DragDropEvent, const FSequencerDragDropParams& DragDropParams) override;
 
-private:
+protected:
+	/** Delegate for AnimatablePropertyChanged in HandleAssetAdded for subtitles */
+	FKeyPropertyResult AddNewSubtitle(FFrameNumber KeyTime, class USubtitleAssetUserData* Subtitle, UMovieSceneSubtitlesTrack* InDestinationTrack, int32 RowIndex);
+	/** Delegate for AnimatablePropertyChanged in HandleAssetAdded for attached subtitles */
+	FKeyPropertyResult AddNewAttachedSubtitle(FFrameNumber KeyTime, class USubtitleAssetUserData* Subtitle, UMovieSceneSubtitlesTrack* InDestinationTrack, TArray<TWeakObjectPtr<UObject>> ObjectsToAttachTo);
 
+
+
+private:
 	TSharedRef<SWidget> BuildSubMenu(FOnAssetSelected OnAssetSelected, FOnAssetEnterPressed OnAssetEnterPressed);
 	void OnAssetSelected(const FAssetData& AssetData, UMovieSceneTrack* Track);
 	void OnAssetEnterPressed(const TArray<FAssetData>& AssetData, UMovieSceneTrack* Track);
 	void HandleAddMenuEntryExecute();
+
+	// Helper function: packages up and null-checks common info used by OnAllowDrop and OnDrop
+	bool IsDragDropValid(const FDragDropEvent& DragDropEvent, const FSequencerDragDropParams& DragDropParams, DragDropValidInfo& OutDragDropInfo);
 };
diff --git a/Engine/Plugins/Experimental/SubtitlesAndClosedCaptions/Source/SubtitlesAndClosedCaptionsEditor/SubtitlesAndClosedCaptionsEditor.Build.cs b/Engine/Plugins/Experimental/SubtitlesAndClosedCaptions/Source/SubtitlesAndClosedCaptionsEditor/SubtitlesAndClosedCaptionsEditor.Build.cs
--- a/Engine/Plugins/Experimental/SubtitlesAndClosedCaptions/Source/SubtitlesAndClosedCaptionsEditor/SubtitlesAndClosedCaptionsEditor.Build.cs
+++ b/Engine/Plugins/Experimental/SubtitlesAndClosedCaptions/Source/SubtitlesAndClosedCaptionsEditor/SubtitlesAndClosedCaptionsEditor.Build.cs
@@ -16,6 +16,7 @@ namespace UnrealBuildTool.Rules
 					"Blutility",
 					"LevelSequence",
 					"MovieScene",
+					"MovieSceneTools",
 					"Overlay",
 					"Sequencer",
 					"SequencerCore",
