Subtitles in an array should now play at correct times when queued in Sequences.

diff --git a/Engine/Plugins/Experimental/SubtitlesAndClosedCaptions/Source/SubtitlesAndClosedCaptions/Public/MovieSceneSubtitlesSystem.h b/Engine/Plugins/Experimental/SubtitlesAndClosedCaptions/Source/SubtitlesAndClosedCaptions/Public/MovieSceneSubtitlesSystem.h
--- a/Engine/Plugins/Experimental/SubtitlesAndClosedCaptions/Source/SubtitlesAndClosedCaptions/Public/MovieSceneSubtitlesSystem.h
+++ b/Engine/Plugins/Experimental/SubtitlesAndClosedCaptions/Source/SubtitlesAndClosedCaptions/Public/MovieSceneSubtitlesSystem.h
@@ -92,21 +92,26 @@ struct FEvaluateSubtitles
 			const FFrameNumber FirstSubtitleFrame = SubtitleSection->GetRange().GetLowerBoundValue();
 			const FFrameNumber LastEvaluatedFrame = Context.GetFrameNumberRange().GetUpperBoundValue();
 
+			// Only try to play subtitles if inside the section's first and last frames.
 			if (LastEvaluatedFrame >= FirstSubtitleFrame && LastEvaluatedFrame < LastSubtitleFrame)
 			{
 				const FFrameRate FrameRate = Context.GetFrameRate();
 				for (const FSubtitleAssetData& Entry : SubtitleAsset->Subtitles)
 				{
+					// The subtitle asset may have several subtitles.
+					// To ensure each is queued at the correct time, check the start and end times of each against the current playback position.
 					const FFrameTime StartOffset = Entry.StartOffset * FrameRate;
 					const FFrameNumber FirstOffsetSubtitleFrame = FirstSubtitleFrame + StartOffset.FrameNumber;
 
-					// Only need to test for the initial offset for the beginning; the end of the section includes the delay when it is first initialized.
-					if (LastEvaluatedFrame >= FirstOffsetSubtitleFrame)
+					const FFrameTime LastFrameTime = Entry.Duration * FrameRate;
+					const FFrameNumber LastOffsetSubtitleFrame = FirstOffsetSubtitleFrame + LastFrameTime.FrameNumber;
+
+					if (LastEvaluatedFrame >= FirstOffsetSubtitleFrame && LastEvaluatedFrame < LastOffsetSubtitleFrame)
 					{
-						// Queue the subtitle to display in short (1/4 second) increments whenever playback is over the subtitle.
+						// Queue the subtitle to display in short (1/10th second) increments whenever playback is over the subtitle.
 						// This allows for jumping and scrubbing the sequence without the subtitle getting stuck on or off.
 						FQueueSubtitleParameters Params{ Entry };
-						Params.Duration = 0.25f;
+						Params.Duration = 0.1f;
 						Params.StartOffset = 0.f;
 
 						FSubtitlesAndClosedCaptionsDelegates::QueueSubtitle.ExecuteIfBound(Params, ESubtitleTiming::InternallyTimed);
