[Subtitles] Fix overzealous check() crashes from empty Dialogue nodes in SoundCues.

diff --git a/Engine/Plugins/Experimental/SubtitlesAndClosedCaptions/Source/SubtitlesAndClosedCaptions/Private/SubtitlesAudioSubsystem.cpp b/Engine/Plugins/Experimental/SubtitlesAndClosedCaptions/Source/SubtitlesAndClosedCaptions/Private/SubtitlesAudioSubsystem.cpp
--- a/Engine/Plugins/Experimental/SubtitlesAndClosedCaptions/Source/SubtitlesAndClosedCaptions/Private/SubtitlesAudioSubsystem.cpp
+++ b/Engine/Plugins/Experimental/SubtitlesAndClosedCaptions/Source/SubtitlesAndClosedCaptions/Private/SubtitlesAudioSubsystem.cpp
@@ -87,13 +87,17 @@ void USubtitlesAudioSubsystem::NotifyActiveSoundDeleting(const FActiveSound& Act
 bool USubtitlesAudioSubsystem::CheckSoundCueNodeForDialogueWaves(const USoundNode* Node, const FActiveSound& ActiveSound)
 {
 	const USoundNodeDialoguePlayer* DialoguePlayerNode = Cast<USoundNodeDialoguePlayer>(Node);
+
+	// These pointers might be null or invalid:
+	// A SoundCue might have a Dialogue node that doesn't actually reference any assets, or the asset itself might lack any actual soundwaves to play.
 	if (IsValid(DialoguePlayerNode))
 	{
 		const UDialogueSoundWaveProxy* Proxy = DialoguePlayerNode->GetContextualDialogueSoundWaveProxy();
-		check(Proxy);
-
-		QueueSubtitleFromDialogueWaveProxy(Proxy, ActiveSound);
-		return true;
+		if (IsValid(Proxy))
+		{
+			QueueSubtitleFromDialogueWaveProxy(Proxy, ActiveSound);
+			return true;
+		}
 	}
 
 	return false;
diff --git a/Engine/Source/Runtime/Engine/Private/SoundNodeDialoguePlayer.cpp b/Engine/Source/Runtime/Engine/Private/SoundNodeDialoguePlayer.cpp
--- a/Engine/Source/Runtime/Engine/Private/SoundNodeDialoguePlayer.cpp
+++ b/Engine/Source/Runtime/Engine/Private/SoundNodeDialoguePlayer.cpp
@@ -95,16 +95,17 @@ void USoundNodeDialoguePlayer::SetDialogueWave(UDialogueWave* Value)
 const UDialogueSoundWaveProxy* USoundNodeDialoguePlayer::GetContextualDialogueSoundWaveProxy() const
 {
 	const UDialogueWave* DialogueWave = GetDialogueWave();
-	check(IsValid(DialogueWave));
 
-	const USoundBase* SoundBase = DialogueWave ? DialogueWave->GetWaveFromContext(DialogueWaveParameter.Context) : NULL;
-	if (IsValid(SoundBase))
+	if (IsValid(DialogueWave))
 	{
-		const UDialogueSoundWaveProxy* Proxy = Cast<UDialogueSoundWaveProxy>(SoundBase);
-		check(IsValid(Proxy));	// GetWaveFromContext() specifically uses DialogueSoundWaveProxy, but casts to SoundBase before returning.
-		return Proxy;
+		const USoundBase* SoundBase = DialogueWave ? DialogueWave->GetWaveFromContext(DialogueWaveParameter.Context) : NULL;
+		if (IsValid(SoundBase))
+		{
+			const UDialogueSoundWaveProxy* Proxy = Cast<UDialogueSoundWaveProxy>(SoundBase);
+			check(IsValid(Proxy));	// GetWaveFromContext() specifically uses DialogueSoundWaveProxy, but casts to SoundBase before returning.
+			return Proxy;
+		}
 	}
-
 	return nullptr;
 }
 
