Add a Blueprint node for queueing all subtitles in a USubtitleAssetUserData.

diff --git a/Engine/Plugins/Experimental/SubtitlesAndClosedCaptions/Source/SubtitlesAndClosedCaptions/Private/SubtitlesBlueprintFunctionLibrary.cpp b/Engine/Plugins/Experimental/SubtitlesAndClosedCaptions/Source/SubtitlesAndClosedCaptions/Private/SubtitlesBlueprintFunctionLibrary.cpp
--- a/Engine/Plugins/Experimental/SubtitlesAndClosedCaptions/Source/SubtitlesAndClosedCaptions/Private/SubtitlesBlueprintFunctionLibrary.cpp
+++ b/Engine/Plugins/Experimental/SubtitlesAndClosedCaptions/Source/SubtitlesAndClosedCaptions/Private/SubtitlesBlueprintFunctionLibrary.cpp
@@ -6,6 +6,19 @@
 
 #include UE_INLINE_GENERATED_CPP_BY_NAME(SubtitlesBlueprintFunctionLibrary)
 
+void USubtitlesBlueprintFunctionLibrary::QueueSubtitlesFromAsset(const USubtitleAssetUserData* SubtitleAsset, const ESubtitleTiming Timing)
+{
+	if (IsValid(SubtitleAsset))
+	{
+		for (const FSubtitleAssetData& Subtitle : SubtitleAsset->Subtitles)
+		{
+			FQueueSubtitleParameters Params{ Subtitle };
+			FSubtitlesAndClosedCaptionsDelegates::QueueSubtitle.ExecuteIfBound(Params, Timing);
+		}
+	}
+	return;
+}
+
 void USubtitlesBlueprintFunctionLibrary::QueueSubtitle(const FSubtitleAssetData& Subtitle, const ESubtitleTiming Timing)
 {
 	FQueueSubtitleParameters Params{ Subtitle };
diff --git a/Engine/Plugins/Experimental/SubtitlesAndClosedCaptions/Source/SubtitlesAndClosedCaptions/Public/SubtitlesBlueprintFunctionLibrary.h b/Engine/Plugins/Experimental/SubtitlesAndClosedCaptions/Source/SubtitlesAndClosedCaptions/Public/SubtitlesBlueprintFunctionLibrary.h
--- a/Engine/Plugins/Experimental/SubtitlesAndClosedCaptions/Source/SubtitlesAndClosedCaptions/Public/SubtitlesBlueprintFunctionLibrary.h
+++ b/Engine/Plugins/Experimental/SubtitlesAndClosedCaptions/Source/SubtitlesAndClosedCaptions/Public/SubtitlesBlueprintFunctionLibrary.h
@@ -14,7 +14,11 @@ class USubtitlesBlueprintFunctionLibrary : public UBlueprintFunctionLibrary
 	GENERATED_BODY()
 
 public:
+	// Queues all of the subtitles in a Subtitle asset, accounting for any delayed start offsets set for each.
+	UFUNCTION(BlueprintCallable, Category = Subtitles)
+	static UE_API void QueueSubtitlesFromAsset(const USubtitleAssetUserData* SubtitleAsset, const ESubtitleTiming Timing = ESubtitleTiming::InternallyTimed);
 
+	// Queues an individual subtitle from a struct, or a single entry in a Subtitle asset's array.
 	UFUNCTION(BlueprintCallable, Category = Subtitles)
 	static UE_API void QueueSubtitle(const FSubtitleAssetData& Subtitle, const ESubtitleTiming Timing = ESubtitleTiming::InternallyTimed);
 
