[Subtitles] Ensure plugin subtitles don't autoplay from DialogueWaves without the migration cvar set.

diff --git a/Engine/Plugins/Experimental/SubtitlesAndClosedCaptions/Source/SubtitlesAndClosedCaptions/Private/SubtitlesAudioSubsystem.cpp b/Engine/Plugins/Experimental/SubtitlesAndClosedCaptions/Source/SubtitlesAndClosedCaptions/Private/SubtitlesAudioSubsystem.cpp
--- a/Engine/Plugins/Experimental/SubtitlesAndClosedCaptions/Source/SubtitlesAndClosedCaptions/Private/SubtitlesAudioSubsystem.cpp
+++ b/Engine/Plugins/Experimental/SubtitlesAndClosedCaptions/Source/SubtitlesAndClosedCaptions/Private/SubtitlesAudioSubsystem.cpp
@@ -27,7 +27,11 @@ void USubtitlesAudioSubsystem::NotifyActiveSoundCreated(FActiveSound& ActiveSoun
 		}
 
 		// Also check for DialogueWaves with localized text
-		if (ActiveSound.bHandleSubtitles)
+		// DialogueWaves, even inside SoundCues, should only play subtitles this way if the migration cvar was set.
+		// Since this callback is run when the sound is created in the sound engine, changing the cvar doesn't effect currently-playing sounds.
+		static const IConsoleVariable* CVarUseNewSubtitles = IConsoleManager::Get().FindConsoleVariable(TEXT("au.UseNewSubtitles"));
+
+		if (CVarUseNewSubtitles && CVarUseNewSubtitles->GetBool() && ActiveSound.bHandleSubtitles)
 		{
 			FSubtitleAssetData SubtitleCue;
 
diff --git a/Engine/Source/Runtime/Engine/Private/DialogueWave.cpp b/Engine/Source/Runtime/Engine/Private/DialogueWave.cpp
--- a/Engine/Source/Runtime/Engine/Private/DialogueWave.cpp
+++ b/Engine/Source/Runtime/Engine/Private/DialogueWave.cpp
@@ -36,7 +36,7 @@ bool bUseNewSubtitles = false;
 FAutoConsoleVariableRef CVarUseNewSubtitles(
 	TEXT("au.UseNewSubtitles"),
 	bUseNewSubtitles,
-	TEXT("Defines whether to force DialogueWaves to use the new subtitles plugin.\n")
+	TEXT("Defines whether to force DialogueWaves to use the new subtitles plugin. Currently-playing sounds will not have their queued or displayed subtitles updated.\n")
 	TEXT("0: Use old subtitles\n")
 	TEXT("1: Use new subtitles\n"),
 	ECVF_Default
