[Subtitles] Adds a function to USubtitlesSubsystem for hot-swapping the Widget.

diff --git a/Engine/Plugins/Experimental/SubtitlesAndClosedCaptions/Source/SubtitlesAndClosedCaptions/Private/SubtitlesBlueprintFunctionLibrary.cpp b/Engine/Plugins/Experimental/SubtitlesAndClosedCaptions/Source/SubtitlesAndClosedCaptions/Private/SubtitlesBlueprintFunctionLibrary.cpp
--- a/Engine/Plugins/Experimental/SubtitlesAndClosedCaptions/Source/SubtitlesAndClosedCaptions/Private/SubtitlesBlueprintFunctionLibrary.cpp
+++ b/Engine/Plugins/Experimental/SubtitlesAndClosedCaptions/Source/SubtitlesAndClosedCaptions/Private/SubtitlesBlueprintFunctionLibrary.cpp
@@ -2,7 +2,9 @@
 
 #include "SubtitlesBlueprintFunctionLibrary.h"
 
-#include "Subtitles/SubtitlesAndClosedCaptionsDelegates.h"
+#include "Engine/Engine.h"
+#include "Subsystems/SubsystemBlueprintLibrary.h"
+#include "SubtitlesSubsystem.h"
 
 #include UE_INLINE_GENERATED_CPP_BY_NAME(SubtitlesBlueprintFunctionLibrary)
 
@@ -43,3 +45,31 @@ void USubtitlesBlueprintFunctionLibrary::StopAllSubtitles()
 {
 	FSubtitlesAndClosedCaptionsDelegates::StopAllSubtitles.ExecuteIfBound();
 }
+
+void USubtitlesBlueprintFunctionLibrary::ReplaceSubtitleWidget(const TSubclassOf<USubtitleWidget>& NewWidgetAsset)
+{
+	// ReplaceSubtitleWidget doesn't have a corresponding delegate call (no need to expose it to the engine), so we need to get the subsystem manually as a result.
+	if (GEngine == nullptr)
+	{
+		return;
+	}
+
+	// Get the subtitles subsystem from the visible world.
+	if (const UGameViewportClient* Viewport = GEngine->GameViewport)
+	{
+		const FWorldContext& WorldContext = GEngine->GetWorldContextFromGameViewportChecked(Viewport);
+		UWorld* World = WorldContext.World();
+
+		if (IsValid(World))
+		{
+			// ReplaceWidget is non-const, so no const pointers here.
+			UWorldSubsystem* WorldSubsystem = USubsystemBlueprintLibrary::GetWorldSubsystem(World, USubtitlesSubsystem::StaticClass());
+			USubtitlesSubsystem* SubtitlesSubsystem = Cast<USubtitlesSubsystem>(WorldSubsystem);
+
+			if (IsValid(SubtitlesSubsystem))
+			{
+				SubtitlesSubsystem->ReplaceWidget(NewWidgetAsset);
+			}
+		}
+	}
+}
diff --git a/Engine/Plugins/Experimental/SubtitlesAndClosedCaptions/Source/SubtitlesAndClosedCaptions/Private/SubtitlesSubsystem.cpp b/Engine/Plugins/Experimental/SubtitlesAndClosedCaptions/Source/SubtitlesAndClosedCaptions/Private/SubtitlesSubsystem.cpp
--- a/Engine/Plugins/Experimental/SubtitlesAndClosedCaptions/Source/SubtitlesAndClosedCaptions/Private/SubtitlesSubsystem.cpp
+++ b/Engine/Plugins/Experimental/SubtitlesAndClosedCaptions/Source/SubtitlesAndClosedCaptions/Private/SubtitlesSubsystem.cpp
@@ -41,7 +41,6 @@ void USubtitlesSubsystem::BindDelegates()
 	}
 }
 
-
 void USubtitlesSubsystem::QueueSubtitle(const FQueueSubtitleParameters& Params, const ESubtitleTiming Timing)
 {
 	// Don't queue subtitles when the engine doesn't have them enabled or forced off. 
@@ -225,6 +224,20 @@ void USubtitlesSubsystem::StopAllSubtitles()
 	}
 }
 
+void USubtitlesSubsystem::ReplaceWidget(const TSubclassOf<USubtitleWidget>& NewWidgetAsset)
+{
+	if (!ActiveSubtitles.IsEmpty() || !DelayedSubtitles.IsEmpty())
+	{
+		UE_LOG(LogSubtitlesAndClosedCaptions, Error, TEXT("Can't replace the subtitle widget: there are still subtitles queued or displaying. Use StopAllSubtitles() first."));
+		return;
+	}
+
+	if (!TryCreateUMGWidgetFromAsset(NewWidgetAsset))
+	{
+		UE_LOG(LogSubtitlesAndClosedCaptions, Error, TEXT("Can't replace the subtitle widget; was a valid asset provided?"));
+	}
+}
+
 void USubtitlesSubsystem::RemoveActiveSubtitle(const FSubtitleAssetData& Subtitle)
 {
 	bool bSuccessfullyRemoved = false;
@@ -265,17 +278,25 @@ void USubtitlesSubsystem::RemoveActiveSubtitle(const FSubtitleAssetData& Subtitl
 	DelayedSubtitles.SetNum(FirstRemovedIndex);
 }
 
+bool USubtitlesSubsystem::TryCreateUMGWidgetFromAsset(const TSubclassOf<USubtitleWidget>& WidgetToUse)
+{
+	if (IsValid(WidgetToUse))
+	{
+		SubtitleWidget = CreateWidget<USubtitleWidget>(GetWorld(), WidgetToUse);
+		bInitializedWidget = false;	// The new widget will need to be added to the viewport.
+	}
+
+	return IsValid(SubtitleWidget);
+}
+
 bool USubtitlesSubsystem::TryCreateUMGWidget()
 {
 	// Set up the UMG widget
 	const USubtitlesSettings* Settings = GetDefault<USubtitlesSettings>();
 	check(Settings != nullptr);
 	const TSubclassOf<USubtitleWidget>& WidgetToUse = Settings->GetWidget();
-	if (IsValid(WidgetToUse))
-	{
-		SubtitleWidget = CreateWidget<USubtitleWidget>(GetWorld(), WidgetToUse);
-	}
-	else
+	
+	if (!TryCreateUMGWidgetFromAsset(WidgetToUse))
 	{
 		// Fallback to default widget (not set by user):
 		const TSubclassOf<USubtitleWidget>& WidgetToUseDefault = Settings->GetWidgetDefault();
@@ -283,9 +304,13 @@ bool USubtitlesSubsystem::TryCreateUMGWidget()
 		{
 			SubtitleWidget = CreateWidget<USubtitleWidget>(GetWorld(), WidgetToUseDefault);
 		}
+		else
+		{
+			UE_LOG(LogSubtitlesAndClosedCaptions, Error, TEXT("The default Subtitle Widget asset isn't valid. Subtitles won't be displayed."));
+		}
 	}
-	bInitializedWidget = false;
 
+	bInitializedWidget = false;
 	return IsValid(SubtitleWidget);
 }
 
diff --git a/Engine/Plugins/Experimental/SubtitlesAndClosedCaptions/Source/SubtitlesAndClosedCaptions/Public/SubtitlesBlueprintFunctionLibrary.h b/Engine/Plugins/Experimental/SubtitlesAndClosedCaptions/Source/SubtitlesAndClosedCaptions/Public/SubtitlesBlueprintFunctionLibrary.h
--- a/Engine/Plugins/Experimental/SubtitlesAndClosedCaptions/Source/SubtitlesAndClosedCaptions/Public/SubtitlesBlueprintFunctionLibrary.h
+++ b/Engine/Plugins/Experimental/SubtitlesAndClosedCaptions/Source/SubtitlesAndClosedCaptions/Public/SubtitlesBlueprintFunctionLibrary.h
@@ -8,6 +8,8 @@
 
 #define UE_API SUBTITLESANDCLOSEDCAPTIONS_API
 
+class USubtitleWidget;
+
 UCLASS(MinimalAPI, meta = (ScriptName = "SubtitlesLibrary"))
 class USubtitlesBlueprintFunctionLibrary : public UBlueprintFunctionLibrary
 {
@@ -30,6 +32,9 @@ public:
 
 	UFUNCTION(BlueprintCallable, Category = Subtitles)
 	static UE_API void StopAllSubtitles();
+
+	UFUNCTION(BlueprintCallable, Category = Subtitles)
+	static UE_API void ReplaceSubtitleWidget(const TSubclassOf<USubtitleWidget>& NewWidgetAsset);
 };
 
 #undef UE_API
diff --git a/Engine/Plugins/Experimental/SubtitlesAndClosedCaptions/Source/SubtitlesAndClosedCaptions/Public/SubtitlesSubsystem.h b/Engine/Plugins/Experimental/SubtitlesAndClosedCaptions/Source/SubtitlesAndClosedCaptions/Public/SubtitlesSubsystem.h
--- a/Engine/Plugins/Experimental/SubtitlesAndClosedCaptions/Source/SubtitlesAndClosedCaptions/Public/SubtitlesSubsystem.h
+++ b/Engine/Plugins/Experimental/SubtitlesAndClosedCaptions/Source/SubtitlesAndClosedCaptions/Public/SubtitlesSubsystem.h
@@ -57,6 +57,8 @@ public:
 	// Stops all queued subtitles from being displayed.  This includes subtitles not yet being displayed due to their StartOffset.
 	UE_API virtual void StopAllSubtitles();
 
+	UE_API virtual void ReplaceWidget(const TSubclassOf<USubtitleWidget>& NewWidgetAsset);
+
 protected:
 
 	UE_API virtual void Initialize(FSubsystemCollectionBase& Collection) override;
@@ -117,6 +119,7 @@ private:
 	TObjectPtr<USubtitleWidget> SubtitleWidget = nullptr;
 
 	bool bInitializedWidget = false;
+	UE_API bool TryCreateUMGWidgetFromAsset(const TSubclassOf<USubtitleWidget>& WidgetToUse);
 	UE_API bool TryCreateUMGWidget();
 	UE_API void UpdateWidgetData();
 };
