Subtitles: Add StopSubtitlesInAsset BP function, remove External timing as an option for QueueSubtitlesFromAsset.

diff --git a/Engine/Plugins/Experimental/SubtitlesAndClosedCaptions/Source/SubtitlesAndClosedCaptions/Private/SubtitlesBlueprintFunctionLibrary.cpp b/Engine/Plugins/Experimental/SubtitlesAndClosedCaptions/Source/SubtitlesAndClosedCaptions/Private/SubtitlesBlueprintFunctionLibrary.cpp
--- a/Engine/Plugins/Experimental/SubtitlesAndClosedCaptions/Source/SubtitlesAndClosedCaptions/Private/SubtitlesBlueprintFunctionLibrary.cpp
+++ b/Engine/Plugins/Experimental/SubtitlesAndClosedCaptions/Source/SubtitlesAndClosedCaptions/Private/SubtitlesBlueprintFunctionLibrary.cpp
@@ -8,17 +8,30 @@
 
 #include UE_INLINE_GENERATED_CPP_BY_NAME(SubtitlesBlueprintFunctionLibrary)
 
-void USubtitlesBlueprintFunctionLibrary::QueueSubtitlesFromAsset(const USubtitleAssetUserData* SubtitleAsset, const ESubtitleTiming Timing)
+void USubtitlesBlueprintFunctionLibrary::QueueSubtitlesFromAsset(TSoftObjectPtr<USubtitleAssetUserData> SubtitleAsset)
 {
-	if (IsValid(SubtitleAsset))
+	if (SubtitleAsset.IsValid())
 	{
 		for (const FSubtitleAssetData& Subtitle : SubtitleAsset->Subtitles)
 		{
 			FQueueSubtitleParameters Params{ Subtitle };
-			FSubtitlesAndClosedCaptionsDelegates::QueueSubtitle.ExecuteIfBound(Params, Timing);
+
+			// Force internal timing: external timing requires the caller to manually start and stop each subtitle in the asset, defeating this function's purpose.
+			// For manual timing control on a per-subtitle basis, QueueSubtitle exists.
+			FSubtitlesAndClosedCaptionsDelegates::QueueSubtitle.ExecuteIfBound(Params, ESubtitleTiming::InternallyTimed);
+		}
+	}
+}
+
+void USubtitlesBlueprintFunctionLibrary::StopSubtitlesInAsset(TSoftObjectPtr<USubtitleAssetUserData> SubtitleAsset)
+{
+	if (SubtitleAsset.IsValid())
+	{
+		for (const FSubtitleAssetData& Subtitle : SubtitleAsset->Subtitles)
+		{
+			FSubtitlesAndClosedCaptionsDelegates::StopSubtitle.ExecuteIfBound(Subtitle);
 		}
 	}
-	return;
 }
 
 void USubtitlesBlueprintFunctionLibrary::QueueSubtitle(const FSubtitleAssetData& Subtitle, const ESubtitleTiming Timing)
diff --git a/Engine/Plugins/Experimental/SubtitlesAndClosedCaptions/Source/SubtitlesAndClosedCaptions/Public/SubtitlesBlueprintFunctionLibrary.h b/Engine/Plugins/Experimental/SubtitlesAndClosedCaptions/Source/SubtitlesAndClosedCaptions/Public/SubtitlesBlueprintFunctionLibrary.h
--- a/Engine/Plugins/Experimental/SubtitlesAndClosedCaptions/Source/SubtitlesAndClosedCaptions/Public/SubtitlesBlueprintFunctionLibrary.h
+++ b/Engine/Plugins/Experimental/SubtitlesAndClosedCaptions/Source/SubtitlesAndClosedCaptions/Public/SubtitlesBlueprintFunctionLibrary.h
@@ -16,9 +16,17 @@ class USubtitlesBlueprintFunctionLibrary : public UBlueprintFunctionLibrary
 	GENERATED_BODY()
 
 public:
-	// Queues all of the subtitles in a Subtitle asset, accounting for any delayed start offsets set for each.
+	// Functions taking USubtitleAssetUserData as a parameter use TSoftObjectPtr to ensure that you can drag an asset onto BP variables feeding into them.
+	// A const pointer allows drag/drop onto the function parameter itself, but TSoftObjectPtr is required for anything with the DefaultToInstanced specifier
+	// (which in this case is the parent class UAssetUserData).
+
+	// Queues all of the subtitles in a Subtitle asset, accounting for any delayed start offsets set for each. Subtitles queued through this function are internally-timed.
+	UFUNCTION(BlueprintCallable, Category = Subtitles)
+	static UE_API void QueueSubtitlesFromAsset(TSoftObjectPtr<USubtitleAssetUserData> SubtitleAsset);
+
+	// Stops all subtitles in a subtitle asset, regardless of whether they're already queued or if they have a delayed start.
 	UFUNCTION(BlueprintCallable, Category = Subtitles)
-	static UE_API void QueueSubtitlesFromAsset(const USubtitleAssetUserData* SubtitleAsset, const ESubtitleTiming Timing = ESubtitleTiming::InternallyTimed);
+	static UE_API void StopSubtitlesInAsset(TSoftObjectPtr<USubtitleAssetUserData> SubtitleAsset);
 
 	// Queues an individual subtitle from a struct, or a single entry in a Subtitle asset's array.
 	UFUNCTION(BlueprintCallable, Category = Subtitles)
