Fix Bughawk complaint: reference to struct in Subtitle Delegates

diff --git a/Engine/Plugins/Experimental/SubtitlesAndClosedCaptions/Source/SubtitlesAndClosedCaptions/Private/SubtitlesSubsystem.cpp b/Engine/Plugins/Experimental/SubtitlesAndClosedCaptions/Source/SubtitlesAndClosedCaptions/Private/SubtitlesSubsystem.cpp
--- a/Engine/Plugins/Experimental/SubtitlesAndClosedCaptions/Source/SubtitlesAndClosedCaptions/Private/SubtitlesSubsystem.cpp
+++ b/Engine/Plugins/Experimental/SubtitlesAndClosedCaptions/Source/SubtitlesAndClosedCaptions/Private/SubtitlesSubsystem.cpp
@@ -42,7 +42,7 @@ void USubtitlesSubsystem::BindDelegates()
 }
 
 
-void USubtitlesSubsystem::QueueSubtitle(const FQueueSubtitleParameters& Params, const ESubtitleTiming Timing)
+void USubtitlesSubsystem::QueueSubtitle(FQueueSubtitleParameters Params, const ESubtitleTiming Timing)
 {
 	// Don't queue subtitles when the engine doesn't have them enabled or forced off. 
 	if (GEngine == nullptr)
@@ -177,7 +177,7 @@ void USubtitlesSubsystem::AddAndDisplaySubtitle(FActiveSubtitle& NewActiveSubtit
 	}
 }
 
-bool USubtitlesSubsystem::IsSubtitleActive(const FSubtitleAssetData& Data) const
+bool USubtitlesSubsystem::IsSubtitleActive(FSubtitleAssetData Data) const
 {
 	if (!ensureMsgf(IsInGameThread(), TEXT("IsSubtitleActive must currently be run on the GameThread - ActiveSubtitles vector is not locked")))
 	{
@@ -190,7 +190,7 @@ bool USubtitlesSubsystem::IsSubtitleActive(const FSubtitleAssetData& Data) const
 	return FoundActiveSubtitle != nullptr;
 }
 
-void USubtitlesSubsystem::StopSubtitle(const FSubtitleAssetData& Data)
+void USubtitlesSubsystem::StopSubtitle(FSubtitleAssetData Data)
 {
 	RemoveActiveSubtitle(Data);
 }
diff --git a/Engine/Plugins/Experimental/SubtitlesAndClosedCaptions/Source/SubtitlesAndClosedCaptions/Public/SubtitlesSubsystem.h b/Engine/Plugins/Experimental/SubtitlesAndClosedCaptions/Source/SubtitlesAndClosedCaptions/Public/SubtitlesSubsystem.h
--- a/Engine/Plugins/Experimental/SubtitlesAndClosedCaptions/Source/SubtitlesAndClosedCaptions/Public/SubtitlesSubsystem.h
+++ b/Engine/Plugins/Experimental/SubtitlesAndClosedCaptions/Source/SubtitlesAndClosedCaptions/Public/SubtitlesSubsystem.h
@@ -46,13 +46,13 @@ public:
 	// Adds a subtitle to the queue: Params contains the subtitle asset and an optional duration. The highest-priority subtitle in the queue will be displayed.
 	// If Timing is ExternallyTimed, the queued subtitle will remain in the queue until manually removed.
 	// If the subtitle asset has a non-zero StartOffset, it will sit in a delayed-start queue instead of being queued for display.
-	UE_API virtual void QueueSubtitle(const FQueueSubtitleParameters& Params, const ESubtitleTiming Timing = ESubtitleTiming::InternallyTimed);
+	UE_API virtual void QueueSubtitle(FQueueSubtitleParameters Params, const ESubtitleTiming Timing = ESubtitleTiming::InternallyTimed);
 
 	// Returns true if the given subtitle asset is being displayed.
-	UE_API virtual bool IsSubtitleActive(const FSubtitleAssetData& Data) const;
+	UE_API virtual bool IsSubtitleActive(FSubtitleAssetData Data) const;
 
 	// Stops the given subtitle asset being displayed.  This includes subtitles not yet being displayed due to their StartOffset
-	UE_API virtual void StopSubtitle(const FSubtitleAssetData& Data);
+	UE_API virtual void StopSubtitle(FSubtitleAssetData Data);
 
 	// Stops all queued subtitles from being displayed.  This includes subtitles not yet being displayed due to their StartOffset.
 	UE_API virtual void StopAllSubtitles();
diff --git a/Engine/Source/Runtime/Engine/Private/Subtitles/SubtitlesAndClosedCaptionsDelegates.cpp b/Engine/Source/Runtime/Engine/Private/Subtitles/SubtitlesAndClosedCaptionsDelegates.cpp
--- a/Engine/Source/Runtime/Engine/Private/Subtitles/SubtitlesAndClosedCaptionsDelegates.cpp
+++ b/Engine/Source/Runtime/Engine/Private/Subtitles/SubtitlesAndClosedCaptionsDelegates.cpp
@@ -4,7 +4,7 @@
 
 #include UE_INLINE_GENERATED_CPP_BY_NAME(SubtitlesAndClosedCaptionsDelegates)
 
-TDelegate<void(const FQueueSubtitleParameters&, const ESubtitleTiming)> FSubtitlesAndClosedCaptionsDelegates::QueueSubtitle;
-TDelegate<bool(const FSubtitleAssetData&)> FSubtitlesAndClosedCaptionsDelegates::IsSubtitleActive;
-TDelegate<void(const FSubtitleAssetData&)> FSubtitlesAndClosedCaptionsDelegates::StopSubtitle;
+TDelegate<void(FQueueSubtitleParameters, const ESubtitleTiming)> FSubtitlesAndClosedCaptionsDelegates::QueueSubtitle;
+TDelegate<bool(FSubtitleAssetData)> FSubtitlesAndClosedCaptionsDelegates::IsSubtitleActive;
+TDelegate<void(FSubtitleAssetData)> FSubtitlesAndClosedCaptionsDelegates::StopSubtitle;
 TDelegate<void()> FSubtitlesAndClosedCaptionsDelegates::StopAllSubtitles;
diff --git a/Engine/Source/Runtime/Engine/Public/Subtitles/SubtitlesAndClosedCaptionsDelegates.h b/Engine/Source/Runtime/Engine/Public/Subtitles/SubtitlesAndClosedCaptionsDelegates.h
--- a/Engine/Source/Runtime/Engine/Public/Subtitles/SubtitlesAndClosedCaptionsDelegates.h
+++ b/Engine/Source/Runtime/Engine/Public/Subtitles/SubtitlesAndClosedCaptionsDelegates.h
@@ -127,7 +127,7 @@ public:
 
 struct FQueueSubtitleParameters
 {
-	const FSubtitleAssetData& Subtitle;
+	const FSubtitleAssetData Subtitle;
 	TOptional<float> Duration;
 	TOptional<float> StartOffset;	// Scrubbing subtitles in the sequencer needs an override for StartOffset.
 };
@@ -137,11 +137,11 @@ class FSubtitlesAndClosedCaptionsDelegates
 public:
 
 	// Have the subtitle subsystem to queue a subtitle to be displayed
-	static UE_API TDelegate<void(const FQueueSubtitleParameters&, const ESubtitleTiming)> QueueSubtitle;
+	static UE_API TDelegate<void(FQueueSubtitleParameters, const ESubtitleTiming)> QueueSubtitle;
 
-	static UE_API TDelegate<bool(const FSubtitleAssetData&)> IsSubtitleActive;
+	static UE_API TDelegate<bool(FSubtitleAssetData)> IsSubtitleActive;
 
-	static UE_API TDelegate<void(const FSubtitleAssetData&)> StopSubtitle;
+	static UE_API TDelegate<void(FSubtitleAssetData)> StopSubtitle;
 
 	static UE_API TDelegate<void()> StopAllSubtitles;
 };
