Additional Subtitle Functional Tests
#virtualized

diff --git a/Engine/Plugins/Experimental/SubtitlesAndClosedCaptions/Source/SubtitlesAndClosedCaptions/Public/SubtitlesSubsystem.h b/Engine/Plugins/Experimental/SubtitlesAndClosedCaptions/Source/SubtitlesAndClosedCaptions/Public/SubtitlesSubsystem.h
--- a/Engine/Plugins/Experimental/SubtitlesAndClosedCaptions/Source/SubtitlesAndClosedCaptions/Public/SubtitlesSubsystem.h
+++ b/Engine/Plugins/Experimental/SubtitlesAndClosedCaptions/Source/SubtitlesAndClosedCaptions/Public/SubtitlesSubsystem.h
@@ -88,8 +88,6 @@ protected:
 private:
 	UE_API void AddAndDisplaySubtitle(FActiveSubtitle& NewActiveSubtitle);
 
-#if WITH_DEV_AUTOMATION_TESTS
-	// These functions are for testing subtitle functionality and aren't intended for general use.
 public:
 	// GetActiveSubtitles returns the sorted array of queued subtitles. Higher priorities are earlier in the array.
 	// For equally-prioritized subtitles, the oldest one is first in this array.
@@ -98,6 +96,8 @@ public:
 		return ActiveSubtitles;
 	}
 
+#if WITH_DEV_AUTOMATION_TESTS
+	// These functions are for testing subtitle functionality and aren't intended for general use.
 	// As this is for testing, assume that the number of subtitles has already been checked.
 	const FSubtitleAssetData& GetTopRankedSubtitle() const
 	{
@@ -114,6 +114,11 @@ public:
 #endif
 	//	WITH_DEV_AUTOMATION_TESTS
 
+	inline TWeakObjectPtr<USubtitleWidget> GetActiveSubtitleWidget()
+	{
+		return SubtitleWidget;
+	}
+
 private:
 	UPROPERTY()
 	TObjectPtr<USubtitleWidget> SubtitleWidget = nullptr;
