Manually merge 46537265 from Release-38.10 to Main.

diff --git a/Engine/Plugins/Experimental/SubtitlesAndClosedCaptions/Source/SubtitlesAndClosedCaptions/Private/SubtitlesSubsystem.cpp b/Engine/Plugins/Experimental/SubtitlesAndClosedCaptions/Source/SubtitlesAndClosedCaptions/Private/SubtitlesSubsystem.cpp
--- a/Engine/Plugins/Experimental/SubtitlesAndClosedCaptions/Source/SubtitlesAndClosedCaptions/Private/SubtitlesSubsystem.cpp
+++ b/Engine/Plugins/Experimental/SubtitlesAndClosedCaptions/Source/SubtitlesAndClosedCaptions/Private/SubtitlesSubsystem.cpp
@@ -89,14 +89,14 @@ void USubtitlesSubsystem::QueueSubtitle(FQueueSubtitleParameters Params, const E
 void USubtitlesSubsystem::AddActiveSubtitle(const FSubtitleAssetData& Subtitle, float Duration, const float StartOffset, const ESubtitleTiming Timing)
 {
 	// If the subtitle is already active then update its duration (by removing it then and then re-adding it)
-	const FActiveSubtitle* FoundActiveSubtitle = ActiveSubtitles.FindByPredicate(
-		[&Subtitle](const FActiveSubtitle& ActiveSubtitle) { return ActiveSubtitle.Subtitle == Subtitle; });
+	const FActiveSubtitleAssetData* FoundActiveSubtitle = ActiveSubtitles.FindByPredicate(
+		[&Subtitle](const FActiveSubtitleAssetData& ActiveSubtitle) { return ActiveSubtitle.Subtitle == Subtitle; });
 	if (FoundActiveSubtitle != nullptr)
 	{
 		RemoveActiveSubtitle(Subtitle);
 	}
 
-	FActiveSubtitle NewActiveSubtitle{ Subtitle };
+	FActiveSubtitleAssetData NewActiveSubtitle{ Subtitle };
 
 	// Subtitles with delayed offset need a timer to await their entry in the queue.
 	if (StartOffset > 0.f)
@@ -127,7 +127,7 @@ void USubtitlesSubsystem::AddActiveSubtitle(const FSubtitleAssetData& Subtitle,
 
 void USubtitlesSubsystem::MakeDelayedSubtitleActive(const FSubtitleAssetData& Subtitle, const ESubtitleTiming Timing)
 {
-	FActiveSubtitle* DelayedSubtitle = Algo::FindByPredicate(DelayedSubtitles, [&Subtitle](const FActiveSubtitle& DelayedSubtitle) {return DelayedSubtitle.Subtitle == Subtitle; });
+	FActiveSubtitleAssetData* DelayedSubtitle = Algo::FindByPredicate(DelayedSubtitles, [&Subtitle](const FActiveSubtitleAssetData& DelayedSubtitle) {return DelayedSubtitle.Subtitle == Subtitle; });
 
 	if (DelayedSubtitle != nullptr)
 	{
@@ -146,18 +146,18 @@ void USubtitlesSubsystem::MakeDelayedSubtitleActive(const FSubtitleAssetData& Su
 		AddAndDisplaySubtitle(*DelayedSubtitle);
 
 		// Remove from the list of Delayed Subtitles
-		const int32 FirstRemovedIndex = Algo::RemoveIf(DelayedSubtitles, [&Subtitle](const FActiveSubtitle& DelayedSubtitle) {return DelayedSubtitle.Subtitle == Subtitle; });
+		const int32 FirstRemovedIndex = Algo::RemoveIf(DelayedSubtitles, [&Subtitle](const FActiveSubtitleAssetData& DelayedSubtitle) {return DelayedSubtitle.Subtitle == Subtitle; });
 		DelayedSubtitles.SetNum(FirstRemovedIndex);
 	}
 }
 
-void USubtitlesSubsystem::AddAndDisplaySubtitle(FActiveSubtitle& NewActiveSubtitle)
+void USubtitlesSubsystem::AddAndDisplaySubtitle(FActiveSubtitleAssetData& NewActiveSubtitle)
 {
 	const FSubtitleAssetData& NewSubtitle = NewActiveSubtitle.Subtitle;
 
 	bool ShouldAddSubtitle = true;
 	// Warn if the subtitle won't play due to being lower-priority than the currently-playing subtitle in that category.
-	for (const FActiveSubtitle& PreQueuedSubtitle : ActiveSubtitles)
+	for (const FActiveSubtitleAssetData& PreQueuedSubtitle : ActiveSubtitles)
 	{
 		const FSubtitleAssetData& PreQueuedSubtitleData = PreQueuedSubtitle.Subtitle;
 		if (PreQueuedSubtitleData.SubtitleType == NewSubtitle.SubtitleType && (PreQueuedSubtitleData.Priority >= NewSubtitle.Priority))
@@ -170,7 +170,7 @@ void USubtitlesSubsystem::AddAndDisplaySubtitle(FActiveSubtitle& NewActiveSubtit
 	if (ShouldAddSubtitle)
 	{
 		ActiveSubtitles.Add(MoveTemp(NewActiveSubtitle));
-		ActiveSubtitles.StableSort([](const FActiveSubtitle& Lhs, const FActiveSubtitle& Rhs) { return Lhs.Subtitle.Priority > Rhs.Subtitle.Priority; });
+		ActiveSubtitles.StableSort([](const FActiveSubtitleAssetData& Lhs, const FActiveSubtitleAssetData& Rhs) { return Lhs.Subtitle.Priority > Rhs.Subtitle.Priority; });
 
 		UpdateWidgetData();
 	}
@@ -183,8 +183,8 @@ bool USubtitlesSubsystem::IsSubtitleActive(FSubtitleAssetData Data) const
 		return false;
 	}
 
-	const FActiveSubtitle* FoundActiveSubtitle = ActiveSubtitles.FindByPredicate(
-		[Data](const FActiveSubtitle& ActiveSubtitle) { return ActiveSubtitle.Subtitle == Data; });
+	const FActiveSubtitleAssetData* FoundActiveSubtitle = ActiveSubtitles.FindByPredicate(
+		[Data](const FActiveSubtitleAssetData& ActiveSubtitle) { return ActiveSubtitle.Subtitle == Data; });
 
 	return FoundActiveSubtitle != nullptr;
 }
@@ -198,14 +198,14 @@ void USubtitlesSubsystem::StopAllSubtitles()
 {
 	// Clean up queued subtitles.
 	FTimerManager& TimerManager = GetWorldRef().GetTimerManager();
-	for (FActiveSubtitle& ActiveSubtitle : ActiveSubtitles)
+	for (FActiveSubtitleAssetData& ActiveSubtitle : ActiveSubtitles)
 	{
 		TimerManager.ClearTimer(ActiveSubtitle.DurationTimerHandle);
 	}
 	ActiveSubtitles.Empty();
 
 	// Also remove delayed-start subtitles not yet in the queue.
-	for (FActiveSubtitle& DelayedSubtitle : DelayedSubtitles)
+	for (FActiveSubtitleAssetData& DelayedSubtitle : DelayedSubtitles)
 	{
 		TimerManager.ClearTimer(DelayedSubtitle.DurationTimerHandle);
 	}
@@ -241,7 +241,7 @@ void USubtitlesSubsystem::ReplaceWidget(const TSubclassOf<USubtitleWidget>& NewW
 void USubtitlesSubsystem::RemoveActiveSubtitle(const FSubtitleAssetData& Subtitle)
 {
 	bool bSuccessfullyRemoved = false;
-	int32 FirstRemovedIndex = Algo::StableRemoveIf(ActiveSubtitles, [&Subtitle](const FActiveSubtitle& ActiveSubtitle) {return ActiveSubtitle.Subtitle == Subtitle; });
+	int32 FirstRemovedIndex = Algo::StableRemoveIf(ActiveSubtitles, [&Subtitle](const FActiveSubtitleAssetData& ActiveSubtitle) {return ActiveSubtitle.Subtitle == Subtitle; });
 
 	FTimerManager& TimerManager = GetWorldRef().GetTimerManager();
 	for (int32 Index = FirstRemovedIndex; Index < ActiveSubtitles.Num(); ++Index)
@@ -269,7 +269,7 @@ void USubtitlesSubsystem::RemoveActiveSubtitle(const FSubtitleAssetData& Subtitl
 	}
 
 	// Also remove delayed-start subtitles using this asset.
-	FirstRemovedIndex = Algo::RemoveIf(DelayedSubtitles, [&Subtitle](const FActiveSubtitle& DelayedSubtitle) {return DelayedSubtitle.Subtitle == Subtitle; });
+	FirstRemovedIndex = Algo::RemoveIf(DelayedSubtitles, [&Subtitle](const FActiveSubtitleAssetData& DelayedSubtitle) {return DelayedSubtitle.Subtitle == Subtitle; });
 
 	for (int32 Index = FirstRemovedIndex; Index < DelayedSubtitles.Num(); ++Index)
 	{
diff --git a/Engine/Plugins/Experimental/SubtitlesAndClosedCaptions/Source/SubtitlesAndClosedCaptions/Public/ActiveSubtitle.h b/Engine/Plugins/Experimental/SubtitlesAndClosedCaptions/Source/SubtitlesAndClosedCaptions/Public/ActiveSubtitle.h
deleted file mode 100644
diff --git a/Engine/Plugins/Experimental/SubtitlesAndClosedCaptions/Source/SubtitlesAndClosedCaptions/Public/ActiveSubtitleAssetData.h b/Engine/Plugins/Experimental/SubtitlesAndClosedCaptions/Source/SubtitlesAndClosedCaptions/Public/ActiveSubtitleAssetData.h
new file mode 100644
--- /dev/null
+++ b/Engine/Plugins/Experimental/SubtitlesAndClosedCaptions/Source/SubtitlesAndClosedCaptions/Public/ActiveSubtitleAssetData.h
@@ -0,0 +1,22 @@
+// Copyright Epic Games, Inc. All Rights Reserved.
+
+#pragma once 
+
+struct FSubtitleAssetData;
+
+#include "Engine/TimerHandle.h"
+#include "Subtitles/SubtitlesAndClosedCaptionsDelegates.h"
+#include "UObject/ObjectPtr.h"
+
+#include "ActiveSubtitleAssetData.generated.h"
+
+USTRUCT()
+struct FActiveSubtitleAssetData
+{
+	GENERATED_BODY()
+
+	UPROPERTY()
+	FSubtitleAssetData Subtitle;
+
+	FTimerHandle DurationTimerHandle;
+};
diff --git a/Engine/Plugins/Experimental/SubtitlesAndClosedCaptions/Source/SubtitlesAndClosedCaptions/Public/SubtitlesSubsystem.h b/Engine/Plugins/Experimental/SubtitlesAndClosedCaptions/Source/SubtitlesAndClosedCaptions/Public/SubtitlesSubsystem.h
--- a/Engine/Plugins/Experimental/SubtitlesAndClosedCaptions/Source/SubtitlesAndClosedCaptions/Public/SubtitlesSubsystem.h
+++ b/Engine/Plugins/Experimental/SubtitlesAndClosedCaptions/Source/SubtitlesAndClosedCaptions/Public/SubtitlesSubsystem.h
@@ -2,7 +2,7 @@
 
 #pragma once 
 
-#include "ActiveSubtitle.h"
+#include "ActiveSubtitleAssetData.h"
 #include "Fonts/SlateFontInfo.h"
 #include "Math/Color.h"
 #include "Subsystems/WorldSubsystem.h"
@@ -69,11 +69,11 @@ protected:
 
 	// Sorted by priority, desc
 	UPROPERTY()
-	TArray<FActiveSubtitle> ActiveSubtitles;
+	TArray<FActiveSubtitleAssetData> ActiveSubtitles;
 
 	// Unsorted; subtitles with a delayed start offset still need to be tracked before entering the queue proper.
 	UPROPERTY()
-	TArray<FActiveSubtitle> DelayedSubtitles;
+	TArray<FActiveSubtitleAssetData> DelayedSubtitles;
 
 protected:
 	UE_API virtual void AddActiveSubtitle(const FSubtitleAssetData& Subtitle, float Duration, const float StartOffset, const ESubtitleTiming Timing);
@@ -86,34 +86,16 @@ protected:
 	UE_API virtual void RemoveActiveSubtitle(const FSubtitleAssetData& Subtitle);
 
 private:
-	UE_API void AddAndDisplaySubtitle(FActiveSubtitle& NewActiveSubtitle);
+	UE_API void AddAndDisplaySubtitle(FActiveSubtitleAssetData& NewActiveSubtitle);
 
 public:
 	// GetActiveSubtitles returns the sorted array of queued subtitles. Higher priorities are earlier in the array.
 	// For equally-prioritized subtitles, the oldest one is first in this array.
-	const TArray<FActiveSubtitle>& GetActiveSubtitles() const
+	const TArray<FActiveSubtitleAssetData>& GetActiveSubtitles() const
 	{
 		return ActiveSubtitles;
 	}
 
-#if WITH_DEV_AUTOMATION_TESTS
-	// These functions are for testing subtitle functionality and aren't intended for general use.
-	// As this is for testing, assume that the number of subtitles has already been checked.
-	const FSubtitleAssetData& GetTopRankedSubtitle() const
-	{
-		check(ActiveSubtitles.Num() > 0);
-		const FActiveSubtitle& ActiveSubtitle = ActiveSubtitles[0];
-
-		return ActiveSubtitle.Subtitle;
-	}
-
-	void TestActivatingDelayedSubtitle(const FSubtitleAssetData& Data)
-	{
-		MakeDelayedSubtitleActive(Data, ESubtitleTiming::InternallyTimed);
-	}
-#endif
-	//	WITH_DEV_AUTOMATION_TESTS
-
 	inline TWeakObjectPtr<USubtitleWidget> GetActiveSubtitleWidget()
 	{
 		return SubtitleWidget;
