[Subtitles] add Warnings on failed Subtitle Widget Replacement

diff --git a/Engine/Plugins/Experimental/SubtitlesAndClosedCaptions/Source/SubtitlesAndClosedCaptions/Private/SubtitlesSubsystem.cpp b/Engine/Plugins/Experimental/SubtitlesAndClosedCaptions/Source/SubtitlesAndClosedCaptions/Private/SubtitlesSubsystem.cpp
--- a/Engine/Plugins/Experimental/SubtitlesAndClosedCaptions/Source/SubtitlesAndClosedCaptions/Private/SubtitlesSubsystem.cpp
+++ b/Engine/Plugins/Experimental/SubtitlesAndClosedCaptions/Source/SubtitlesAndClosedCaptions/Private/SubtitlesSubsystem.cpp
@@ -285,6 +285,10 @@ bool USubtitlesSubsystem::TryCreateUMGWidgetFromAsset(const TSubclassOf<USubtitl
 		SubtitleWidget = CreateWidget<USubtitleWidget>(GetWorld(), WidgetToUse);
 		bInitializedWidget = false;	// The new widget will need to be added to the viewport.
 	}
+	else
+	{
+		UE_LOG(LogSubtitlesAndClosedCaptions, Warning, TEXT("Can't replace the Subtitle Widget: the supplied asset isn't valid. The existing widget will continue to be used."));
+	}
 
 	return IsValid(SubtitleWidget);
 }
@@ -303,6 +307,9 @@ bool USubtitlesSubsystem::TryCreateUMGWidget()
 		if (IsValid(WidgetToUseDefault))
 		{
 			SubtitleWidget = CreateWidget<USubtitleWidget>(GetWorld(), WidgetToUseDefault);
+
+			// Log the fallback as a warning as seeing the default widget probably isn't intended.
+			UE_LOG(LogSubtitlesAndClosedCaptions, Warning, TEXT("A valid Subtitle Widget wasn't provided. The plugin's default widget will be used as a fallback."));
 		}
 		else
 		{
