Subtitles: Convert widget references using TSubclassOf to TSoftClassPtr

diff --git a/Engine/Plugins/Experimental/SubtitlesAndClosedCaptions/Source/SubtitlesAndClosedCaptions/Private/SubtitlesBlueprintFunctionLibrary.cpp b/Engine/Plugins/Experimental/SubtitlesAndClosedCaptions/Source/SubtitlesAndClosedCaptions/Private/SubtitlesBlueprintFunctionLibrary.cpp
--- a/Engine/Plugins/Experimental/SubtitlesAndClosedCaptions/Source/SubtitlesAndClosedCaptions/Private/SubtitlesBlueprintFunctionLibrary.cpp
+++ b/Engine/Plugins/Experimental/SubtitlesAndClosedCaptions/Source/SubtitlesAndClosedCaptions/Private/SubtitlesBlueprintFunctionLibrary.cpp
@@ -59,7 +59,7 @@ void USubtitlesBlueprintFunctionLibrary::StopAllSubtitles()
 	FSubtitlesAndClosedCaptionsDelegates::StopAllSubtitles.ExecuteIfBound();
 }
 
-void USubtitlesBlueprintFunctionLibrary::ReplaceSubtitleWidget(const TSubclassOf<USubtitleWidget>& NewWidgetAsset)
+void USubtitlesBlueprintFunctionLibrary::ReplaceSubtitleWidget(const TSoftClassPtr<USubtitleWidget>& NewWidgetAsset)
 {
 	// ReplaceSubtitleWidget doesn't have a corresponding delegate call (no need to expose it to the engine), so we need to get the subsystem manually as a result.
 	if (GEngine == nullptr)
diff --git a/Engine/Plugins/Experimental/SubtitlesAndClosedCaptions/Source/SubtitlesAndClosedCaptions/Private/SubtitlesSettings.cpp b/Engine/Plugins/Experimental/SubtitlesAndClosedCaptions/Source/SubtitlesAndClosedCaptions/Private/SubtitlesSettings.cpp
--- a/Engine/Plugins/Experimental/SubtitlesAndClosedCaptions/Source/SubtitlesAndClosedCaptions/Private/SubtitlesSettings.cpp
+++ b/Engine/Plugins/Experimental/SubtitlesAndClosedCaptions/Source/SubtitlesAndClosedCaptions/Private/SubtitlesSettings.cpp
@@ -7,6 +7,8 @@
 
 USubtitlesSettings::USubtitlesSettings()
 {
+	WidgetToUseAsSubclass = SubtitleWidgetToUse.Get();
+	DefaultWidgetToUseAsSubclass = SubtitleWidgetToUseDefault.Get();
 }
 
 #undef LOCTEXT_NAMESPACE
diff --git a/Engine/Plugins/Experimental/SubtitlesAndClosedCaptions/Source/SubtitlesAndClosedCaptions/Private/SubtitlesSubsystem.cpp b/Engine/Plugins/Experimental/SubtitlesAndClosedCaptions/Source/SubtitlesAndClosedCaptions/Private/SubtitlesSubsystem.cpp
--- a/Engine/Plugins/Experimental/SubtitlesAndClosedCaptions/Source/SubtitlesAndClosedCaptions/Private/SubtitlesSubsystem.cpp
+++ b/Engine/Plugins/Experimental/SubtitlesAndClosedCaptions/Source/SubtitlesAndClosedCaptions/Private/SubtitlesSubsystem.cpp
@@ -224,7 +224,7 @@ void USubtitlesSubsystem::StopAllSubtitles()
 	}
 }
 
-void USubtitlesSubsystem::ReplaceWidget(const TSubclassOf<USubtitleWidget>& NewWidgetAsset)
+void USubtitlesSubsystem::ReplaceWidget(const TSoftClassPtr<USubtitleWidget>& NewWidgetAsset)
 {
 	if (!ActiveSubtitles.IsEmpty() || !DelayedSubtitles.IsEmpty())
 	{
@@ -278,11 +278,13 @@ void USubtitlesSubsystem::RemoveActiveSubtitle(const FSubtitleAssetData& Subtitl
 	DelayedSubtitles.SetNum(FirstRemovedIndex);
 }
 
-bool USubtitlesSubsystem::TryCreateUMGWidgetFromAsset(const TSubclassOf<USubtitleWidget>& WidgetToUse)
+bool USubtitlesSubsystem::TryCreateUMGWidgetFromAsset(const TSoftClassPtr<USubtitleWidget>& WidgetToUse)
 {
-	if (IsValid(WidgetToUse))
+	// Synchronous load: called directly from ::Initialize when the World (and this subsystem) is created.
+	WidgetToUse.LoadSynchronous();
+	if (WidgetToUse.IsValid())
 	{
-		SubtitleWidget = CreateWidget<USubtitleWidget>(GetWorld(), WidgetToUse);
+		SubtitleWidget = CreateWidget<USubtitleWidget>(GetWorld(), WidgetToUse.Get());
 		bInitializedWidget = false;	// The new widget will need to be added to the viewport.
 	}
 
@@ -294,15 +296,15 @@ bool USubtitlesSubsystem::TryCreateUMGWidget()
 	// Set up the UMG widget
 	const USubtitlesSettings* Settings = GetDefault<USubtitlesSettings>();
 	check(Settings != nullptr);
-	const TSubclassOf<USubtitleWidget>& WidgetToUse = Settings->GetWidget();
+	const TSoftClassPtr<USubtitleWidget>& WidgetToUse = Settings->GetWidgetSoftClassPtr();
 	
 	if (!TryCreateUMGWidgetFromAsset(WidgetToUse))
 	{
 		// Fallback to default widget (not set by user):
-		const TSubclassOf<USubtitleWidget>& WidgetToUseDefault = Settings->GetWidgetDefault();
-		if (IsValid(WidgetToUseDefault))
+		const TSoftClassPtr<USubtitleWidget>& WidgetToUseDefault = Settings->GetWidgetDefaultSoftClassPtr();
+		if (WidgetToUseDefault.IsValid())
 		{
-			SubtitleWidget = CreateWidget<USubtitleWidget>(GetWorld(), WidgetToUseDefault);
+			SubtitleWidget = CreateWidget<USubtitleWidget>(GetWorld(), WidgetToUseDefault.Get());
 		}
 		else
 		{
diff --git a/Engine/Plugins/Experimental/SubtitlesAndClosedCaptions/Source/SubtitlesAndClosedCaptions/Public/SubtitlesBlueprintFunctionLibrary.h b/Engine/Plugins/Experimental/SubtitlesAndClosedCaptions/Source/SubtitlesAndClosedCaptions/Public/SubtitlesBlueprintFunctionLibrary.h
--- a/Engine/Plugins/Experimental/SubtitlesAndClosedCaptions/Source/SubtitlesAndClosedCaptions/Public/SubtitlesBlueprintFunctionLibrary.h
+++ b/Engine/Plugins/Experimental/SubtitlesAndClosedCaptions/Source/SubtitlesAndClosedCaptions/Public/SubtitlesBlueprintFunctionLibrary.h
@@ -42,7 +42,7 @@ public:
 	static UE_API void StopAllSubtitles();
 
 	UFUNCTION(BlueprintCallable, Category = Subtitles)
-	static UE_API void ReplaceSubtitleWidget(const TSubclassOf<USubtitleWidget>& NewWidgetAsset);
+	static UE_API void ReplaceSubtitleWidget(const TSoftClassPtr<USubtitleWidget>& NewWidgetAsset);
 };
 
 #undef UE_API
diff --git a/Engine/Plugins/Experimental/SubtitlesAndClosedCaptions/Source/SubtitlesAndClosedCaptions/Public/SubtitlesSettings.h b/Engine/Plugins/Experimental/SubtitlesAndClosedCaptions/Source/SubtitlesAndClosedCaptions/Public/SubtitlesSettings.h
--- a/Engine/Plugins/Experimental/SubtitlesAndClosedCaptions/Source/SubtitlesAndClosedCaptions/Public/SubtitlesSettings.h
+++ b/Engine/Plugins/Experimental/SubtitlesAndClosedCaptions/Source/SubtitlesAndClosedCaptions/Public/SubtitlesSettings.h
@@ -16,14 +16,23 @@ class USubtitlesSettings : public UDeveloperSettings
 public:
 	USubtitlesSettings();
 
-	const TSubclassOf<USubtitleWidget>& GetWidget() const { return SubtitleWidgetToUse; }
-	const TSubclassOf<USubtitleWidget>& GetWidgetDefault() const { return SubtitleWidgetToUseDefault; }
+	// Temporarily retaining the old GetWidget to not break builds.
+	// Plugin is still experimental, so we don't need to actually deprecate it just yet.
+	const TSubclassOf<USubtitleWidget>& GetWidget() const { return WidgetToUseAsSubclass; };
+	const TSubclassOf<USubtitleWidget>& GetWidgetDefault() const { return DefaultWidgetToUseAsSubclass;	};
+private:
+	TSubclassOf<USubtitleWidget> WidgetToUseAsSubclass;
+	TSubclassOf<USubtitleWidget> DefaultWidgetToUseAsSubclass;
+
+public:
+	inline const TSoftClassPtr<USubtitleWidget>& GetWidgetSoftClassPtr() const { return SubtitleWidgetToUse; }
+	inline const TSoftClassPtr<USubtitleWidget>& GetWidgetDefaultSoftClassPtr() const { return SubtitleWidgetToUse; }
 
 protected:
 
 	UPROPERTY(config, EditDefaultsOnly, Category = Subtitles, AdvancedDisplay)
-	TSubclassOf<USubtitleWidget> SubtitleWidgetToUse;
+	TSoftClassPtr<USubtitleWidget> SubtitleWidgetToUse;
 
 	UPROPERTY(config)
-	TSubclassOf<USubtitleWidget> SubtitleWidgetToUseDefault; // fallback for SubtitleWidgetToUse (not set by user)
+	TSoftClassPtr<USubtitleWidget> SubtitleWidgetToUseDefault; // fallback for SubtitleWidgetToUse (not set by user)
 };
diff --git a/Engine/Plugins/Experimental/SubtitlesAndClosedCaptions/Source/SubtitlesAndClosedCaptions/Public/SubtitlesSubsystem.h b/Engine/Plugins/Experimental/SubtitlesAndClosedCaptions/Source/SubtitlesAndClosedCaptions/Public/SubtitlesSubsystem.h
--- a/Engine/Plugins/Experimental/SubtitlesAndClosedCaptions/Source/SubtitlesAndClosedCaptions/Public/SubtitlesSubsystem.h
+++ b/Engine/Plugins/Experimental/SubtitlesAndClosedCaptions/Source/SubtitlesAndClosedCaptions/Public/SubtitlesSubsystem.h
@@ -57,7 +57,8 @@ public:
 	// Stops all queued subtitles from being displayed.  This includes subtitles not yet being displayed due to their StartOffset.
 	UE_API virtual void StopAllSubtitles();
 
-	UE_API virtual void ReplaceWidget(const TSubclassOf<USubtitleWidget>& NewWidgetAsset);
+	// Replaces the Subtitle Widget with a new asset subclassed from the base USubtitleWidget. Used by the Replace Subtitle Widget BP function.
+	UE_API virtual void ReplaceWidget(const TSoftClassPtr<USubtitleWidget>& NewWidgetAsset);
 
 protected:
 
@@ -106,7 +107,7 @@ private:
 	TObjectPtr<USubtitleWidget> SubtitleWidget = nullptr;
 
 	bool bInitializedWidget = false;
-	UE_API bool TryCreateUMGWidgetFromAsset(const TSubclassOf<USubtitleWidget>& WidgetToUse);
+	UE_API bool TryCreateUMGWidgetFromAsset(const TSoftClassPtr<USubtitleWidget>& WidgetToUse);
 	UE_API bool TryCreateUMGWidget();
 	UE_API void UpdateWidgetData();
 };
