Mark the in-world subtitle widget as transient, to avoid it being saved as part of the world

This also avoids it being included in the localization gather for the world

#rnx

diff --git a/Engine/Plugins/Experimental/SubtitlesAndClosedCaptions/Source/SubtitlesAndClosedCaptions/Private/SubtitlesSubsystem.cpp b/Engine/Plugins/Experimental/SubtitlesAndClosedCaptions/Source/SubtitlesAndClosedCaptions/Private/SubtitlesSubsystem.cpp
--- a/Engine/Plugins/Experimental/SubtitlesAndClosedCaptions/Source/SubtitlesAndClosedCaptions/Private/SubtitlesSubsystem.cpp
+++ b/Engine/Plugins/Experimental/SubtitlesAndClosedCaptions/Source/SubtitlesAndClosedCaptions/Private/SubtitlesSubsystem.cpp
@@ -285,6 +285,10 @@ bool USubtitlesSubsystem::TryCreateUMGWidgetFromAsset(const TSoftClassPtr<USubti
 	if (WidgetToUse.IsValid())
 	{
 		SubtitleWidget = CreateWidget<USubtitleWidget>(GetWorld(), WidgetToUse.Get());
+		if (IsValid(SubtitleWidget))
+		{
+			SubtitleWidget->SetFlags(RF_Transient);
+		}
 		bInitializedWidget = false;	// The new widget will need to be added to the viewport.
 	}
 
@@ -305,6 +309,10 @@ bool USubtitlesSubsystem::TryCreateUMGWidget()
 		if (WidgetToUseDefault.IsValid())
 		{
 			SubtitleWidget = CreateWidget<USubtitleWidget>(GetWorld(), WidgetToUseDefault.Get());
+			if (IsValid(SubtitleWidget))
+			{
+				SubtitleWidget->SetFlags(RF_Transient);
+			}
 		}
 		else
 		{
diff --git a/Engine/Plugins/Experimental/SubtitlesAndClosedCaptions/Source/SubtitlesAndClosedCaptions/Public/SubtitlesSubsystem.h b/Engine/Plugins/Experimental/SubtitlesAndClosedCaptions/Source/SubtitlesAndClosedCaptions/Public/SubtitlesSubsystem.h
--- a/Engine/Plugins/Experimental/SubtitlesAndClosedCaptions/Source/SubtitlesAndClosedCaptions/Public/SubtitlesSubsystem.h
+++ b/Engine/Plugins/Experimental/SubtitlesAndClosedCaptions/Source/SubtitlesAndClosedCaptions/Public/SubtitlesSubsystem.h
@@ -103,7 +103,7 @@ public:
 	}
 
 private:
-	UPROPERTY()
+	UPROPERTY(Transient)
 	TObjectPtr<USubtitleWidget> SubtitleWidget = nullptr;
 
 	bool bInitializedWidget = false;
