Subtitles: Enable resizing of subtitle sections in sequencer

diff --git a/Engine/Plugins/Experimental/SubtitlesAndClosedCaptions/Source/SubtitlesAndClosedCaptionsEditor/Private/SubtitleSequencerSection.h b/Engine/Plugins/Experimental/SubtitlesAndClosedCaptions/Source/SubtitlesAndClosedCaptionsEditor/Private/SubtitleSequencerSection.h
--- a/Engine/Plugins/Experimental/SubtitlesAndClosedCaptions/Source/SubtitlesAndClosedCaptionsEditor/Private/SubtitleSequencerSection.h
+++ b/Engine/Plugins/Experimental/SubtitlesAndClosedCaptions/Source/SubtitlesAndClosedCaptionsEditor/Private/SubtitleSequencerSection.h
@@ -17,6 +17,5 @@ public:
 
 	// ISequencerSection interface
 	virtual FText GetSectionTitle() const override;
-	virtual bool SectionIsResizable() const override { return false; }
 	virtual int32 OnPaintSection(FSequencerSectionPainter& Painter) const override;
 };
