[Subtitles] Add categories to nativized asset actions, recompile uassets to fill missing supported class array.
#virtualized

diff --git a/Engine/Plugins/Experimental/SubtitlesAndClosedCaptions/Source/SubtitlesAndClosedCaptionsEditor/Public/SubtitleAssetActionUtility.h b/Engine/Plugins/Experimental/SubtitlesAndClosedCaptions/Source/SubtitlesAndClosedCaptionsEditor/Public/SubtitleAssetActionUtility.h
--- a/Engine/Plugins/Experimental/SubtitlesAndClosedCaptions/Source/SubtitlesAndClosedCaptionsEditor/Public/SubtitleAssetActionUtility.h
+++ b/Engine/Plugins/Experimental/SubtitlesAndClosedCaptions/Source/SubtitlesAndClosedCaptionsEditor/Public/SubtitleAssetActionUtility.h
@@ -15,7 +15,7 @@ class USubtitleAssetActionUtility : public UAssetActionUtility
 public:
 	USubtitleAssetActionUtility();
 
-	UFUNCTION(CallInEditor)
+	UFUNCTION(CallInEditor, Category="Sound Base|Subtitles")
 	void ConvertBasicOverlaysToSubtitles();
 
 
diff --git a/Engine/Plugins/Experimental/SubtitlesAndClosedCaptions/Source/SubtitlesAndClosedCaptionsEditor/Public/SubtitleMigrationUtility.h b/Engine/Plugins/Experimental/SubtitlesAndClosedCaptions/Source/SubtitlesAndClosedCaptionsEditor/Public/SubtitleMigrationUtility.h
--- a/Engine/Plugins/Experimental/SubtitlesAndClosedCaptions/Source/SubtitlesAndClosedCaptionsEditor/Public/SubtitleMigrationUtility.h
+++ b/Engine/Plugins/Experimental/SubtitlesAndClosedCaptions/Source/SubtitlesAndClosedCaptionsEditor/Public/SubtitleMigrationUtility.h
@@ -19,12 +19,12 @@ private:
 public:
 	USubtitleMigrationUtility();
 
-	UFUNCTION(CallInEditor)
+	UFUNCTION(CallInEditor, Category = "Sound Base|Subtitles")
 	void AddBlankSubtitle();
 
-	UFUNCTION(CallInEditor)
+	UFUNCTION(CallInEditor, Category = "Sound Base|Subtitles")
 	void RemoveLegacySubtitles();
 
-	UFUNCTION(CallInEditor)
+	UFUNCTION(CallInEditor, Category = "Sound Base|Subtitles")
 	void ConvertLegacySubtitles();
 };
