[Subtitles] Restore backed-out warnings.

diff --git a/Engine/Plugins/Experimental/SubtitlesAndClosedCaptions/Source/SubtitlesAndClosedCaptions/Private/SubtitlesSubsystem.cpp b/Engine/Plugins/Experimental/SubtitlesAndClosedCaptions/Source/SubtitlesAndClosedCaptions/Private/SubtitlesSubsystem.cpp
--- a/Engine/Plugins/Experimental/SubtitlesAndClosedCaptions/Source/SubtitlesAndClosedCaptions/Private/SubtitlesSubsystem.cpp
+++ b/Engine/Plugins/Experimental/SubtitlesAndClosedCaptions/Source/SubtitlesAndClosedCaptions/Private/SubtitlesSubsystem.cpp
@@ -313,6 +313,8 @@ bool USubtitlesSubsystem::TryCreateUMGWidget()
 			{
 				SubtitleWidget->SetFlags(RF_Transient);
 			}
+			// Log the fallback as a warning as seeing the default widget probably isn't intended.
+			UE_LOG(LogSubtitlesAndClosedCaptions, Warning, TEXT("A valid Subtitle Widget wasn't provided. The plugin's default widget will be used as a fallback."));
 		}
 		else
 		{
