Quell an ensure that hits when an empty subtitle is added to Sequencer.

diff --git a/Engine/Plugins/Experimental/SubtitlesAndClosedCaptions/Source/SubtitlesAndClosedCaptions/Private/MovieSceneSubtitlesTrack.cpp b/Engine/Plugins/Experimental/SubtitlesAndClosedCaptions/Source/SubtitlesAndClosedCaptions/Private/MovieSceneSubtitlesTrack.cpp
--- a/Engine/Plugins/Experimental/SubtitlesAndClosedCaptions/Source/SubtitlesAndClosedCaptions/Private/MovieSceneSubtitlesTrack.cpp
+++ b/Engine/Plugins/Experimental/SubtitlesAndClosedCaptions/Source/SubtitlesAndClosedCaptions/Private/MovieSceneSubtitlesTrack.cpp
@@ -31,7 +31,7 @@ UMovieSceneSection* UMovieSceneSubtitlesTrack::AddNewSubtitleOnRow(const USubtit
 	// Find the maximum duration of all the subtitles in the UserData, including their delayed start offsets.
 	const FFrameTime DurationToUse = Subtitle.GetMaximumDuration() * FrameRate;
 
-	ensure(DurationToUse.FrameNumber.Value > 0);
+	ensure(DurationToUse.FrameNumber.Value > 0 || Subtitle.Subtitles.IsEmpty());
 	NewSection->InitialPlacementOnRow(SubtitleSections, InStartTime, DurationToUse.FrameNumber.Value, InRowIndex);
 	NewSection->SetSubtitle(Subtitle);
 
