Fix subtitles settings loading the wrong widget when asked for the default.

#Synced-CL Unknown

diff --git a/Engine/Plugins/Experimental/SubtitlesAndClosedCaptions/Source/SubtitlesAndClosedCaptions/Public/SubtitlesSettings.h b/Engine/Plugins/Experimental/SubtitlesAndClosedCaptions/Source/SubtitlesAndClosedCaptions/Public/SubtitlesSettings.h
--- a/Engine/Plugins/Experimental/SubtitlesAndClosedCaptions/Source/SubtitlesAndClosedCaptions/Public/SubtitlesSettings.h
+++ b/Engine/Plugins/Experimental/SubtitlesAndClosedCaptions/Source/SubtitlesAndClosedCaptions/Public/SubtitlesSettings.h
@@ -26,7 +26,7 @@ private:
 
 public:
 	inline const TSoftClassPtr<USubtitleWidget>& GetWidgetSoftClassPtr() const { return SubtitleWidgetToUse; }
-	inline const TSoftClassPtr<USubtitleWidget>& GetWidgetDefaultSoftClassPtr() const { return SubtitleWidgetToUse; }
+	inline const TSoftClassPtr<USubtitleWidget>& GetWidgetDefaultSoftClassPtr() const { return SubtitleWidgetToUseDefault; }
 
 protected:
 
