[Subtitles] Sequencer sections initialized with a blank subtitle are now given a default duration.

#Synced-CL 0

diff --git a/Engine/Plugins/Experimental/SubtitlesAndClosedCaptions/Source/SubtitlesAndClosedCaptions/Private/MovieSceneSubtitlesTrack.cpp b/Engine/Plugins/Experimental/SubtitlesAndClosedCaptions/Source/SubtitlesAndClosedCaptions/Private/MovieSceneSubtitlesTrack.cpp
--- a/Engine/Plugins/Experimental/SubtitlesAndClosedCaptions/Source/SubtitlesAndClosedCaptions/Private/MovieSceneSubtitlesTrack.cpp
+++ b/Engine/Plugins/Experimental/SubtitlesAndClosedCaptions/Source/SubtitlesAndClosedCaptions/Private/MovieSceneSubtitlesTrack.cpp
@@ -11,6 +11,7 @@
 
 #define LOCTEXT_NAMESPACE "MovieSceneSubtitlesTrack"
 
+constexpr float DefaultSubtitleSectionDuration = 3.f;
 
 UMovieSceneSubtitlesTrack::UMovieSceneSubtitlesTrack(const FObjectInitializer& ObjectInitializer)
 	: Super(ObjectInitializer)
@@ -27,14 +28,23 @@ UMovieSceneSection* UMovieSceneSubtitlesTrack::AddNewSubtitleOnRow(const USubtit
 	UMovieSceneSubtitleSection* NewSection = Cast<UMovieSceneSubtitleSection>(CreateNewSection());
 
 	const FFrameRate FrameRate = GetTypedOuter<UMovieScene>()->GetTickResolution();
+	const float SubtitleDuration = Subtitle.GetMaximumDuration();
 
 	// Find the maximum duration of all the subtitles in the UserData, including their delayed start offsets.
-	const FFrameTime DurationToUse = Subtitle.GetMaximumDuration() * FrameRate;
+	const bool UseDefaultDuration = Subtitle.Subtitles.IsEmpty();
+	// Use the default duration if the subtitles array is empty.
+	// If it's not empty, the subtitles themselves have a minimum duration and therefore have a non-zero length.
+	const FFrameTime DurationToUse = UseDefaultDuration ? (DefaultSubtitleSectionDuration * FrameRate) : (SubtitleDuration * FrameRate);
 
-	ensure(DurationToUse.FrameNumber.Value > 0 || Subtitle.Subtitles.IsEmpty());
 	NewSection->InitialPlacementOnRow(SubtitleSections, InStartTime, DurationToUse.FrameNumber.Value, InRowIndex);
 	NewSection->SetSubtitle(Subtitle);
 
+	// If the duration was 0-length, turn orange to make it clear that the duration was defaulted.
+	if (UseDefaultDuration)
+	{
+		NewSection->SetColorTint(FColor::Orange);
+	}
+
 	SubtitleSections.Add(NewSection);
 
 	return NewSection;
diff --git a/Engine/Plugins/Experimental/SubtitlesAndClosedCaptions/Source/SubtitlesAndClosedCaptionsEditor/Private/SubtitleSequencerSection.cpp b/Engine/Plugins/Experimental/SubtitlesAndClosedCaptions/Source/SubtitlesAndClosedCaptionsEditor/Private/SubtitleSequencerSection.cpp
--- a/Engine/Plugins/Experimental/SubtitlesAndClosedCaptions/Source/SubtitlesAndClosedCaptionsEditor/Private/SubtitleSequencerSection.cpp
+++ b/Engine/Plugins/Experimental/SubtitlesAndClosedCaptions/Source/SubtitlesAndClosedCaptionsEditor/Private/SubtitleSequencerSection.cpp
@@ -19,6 +19,10 @@ FText FSubtitleSequencerSection::GetSectionTitle() const
 			{
 				return Subtitles[0].Text;
 			}
+			else
+			{
+				return NSLOCTEXT("FSubtitlesSequencerSection", "EmptySubtitleAsset", "Empty Subtitle");
+			}
 		}
 	}
 
