[Subtitles] - Introduce a BP function to replace QueueSubtitle to remove access to irrelevant DurationType parameters.

#Synced-CL 0

diff --git a/Engine/Plugins/Experimental/SubtitlesAndClosedCaptions/Source/SubtitlesAndClosedCaptions/Private/SubtitlesBlueprintFunctionLibrary.cpp b/Engine/Plugins/Experimental/SubtitlesAndClosedCaptions/Source/SubtitlesAndClosedCaptions/Private/SubtitlesBlueprintFunctionLibrary.cpp
--- a/Engine/Plugins/Experimental/SubtitlesAndClosedCaptions/Source/SubtitlesAndClosedCaptions/Private/SubtitlesBlueprintFunctionLibrary.cpp
+++ b/Engine/Plugins/Experimental/SubtitlesAndClosedCaptions/Source/SubtitlesAndClosedCaptions/Private/SubtitlesBlueprintFunctionLibrary.cpp
@@ -34,12 +34,28 @@ void USubtitlesBlueprintFunctionLibrary::StopSubtitlesInAsset(TSoftObjectPtr<USu
 	}
 }
 
+// Old deprecated version of QueueSingleSubtitle; use that instead.
 void USubtitlesBlueprintFunctionLibrary::QueueSubtitle(const FSubtitleAssetData& Subtitle, const ESubtitleTiming Timing)
 {
 	FQueueSubtitleParameters Params{ Subtitle };
 	FSubtitlesAndClosedCaptionsDelegates::QueueSubtitle.ExecuteIfBound(Params, Timing);
 }
 
+void USubtitlesBlueprintFunctionLibrary::QueueSingleSubtitle(
+	const FText Text, const float Duration, const float StartOffset, const float Priority, const ESubtitleType SubtitleType, const ESubtitleTiming Timing)
+{
+	FSubtitleAssetData Subtitle{
+		.Text = Text,
+		.SubtitleDurationType = ESubtitleDurationType::UseDurationProperty,	// There is no sound to synchronize to, so use the provided duration.
+		.Duration = Duration,
+		.StartOffset = StartOffset,
+		.Priority = Priority,
+		.SubtitleType = SubtitleType,
+	};
+	FQueueSubtitleParameters Params{ Subtitle };
+	FSubtitlesAndClosedCaptionsDelegates::QueueSubtitle.ExecuteIfBound(Params, Timing);
+}
+
 bool USubtitlesBlueprintFunctionLibrary::IsSubtitleActive(const FSubtitleAssetData& Subtitle)
 {
 	if (!FSubtitlesAndClosedCaptionsDelegates::IsSubtitleActive.IsBound())
diff --git a/Engine/Plugins/Experimental/SubtitlesAndClosedCaptions/Source/SubtitlesAndClosedCaptions/Public/SubtitlesBlueprintFunctionLibrary.h b/Engine/Plugins/Experimental/SubtitlesAndClosedCaptions/Source/SubtitlesAndClosedCaptions/Public/SubtitlesBlueprintFunctionLibrary.h
--- a/Engine/Plugins/Experimental/SubtitlesAndClosedCaptions/Source/SubtitlesAndClosedCaptions/Public/SubtitlesBlueprintFunctionLibrary.h
+++ b/Engine/Plugins/Experimental/SubtitlesAndClosedCaptions/Source/SubtitlesAndClosedCaptions/Public/SubtitlesBlueprintFunctionLibrary.h
@@ -29,9 +29,17 @@ public:
 	static UE_API void StopSubtitlesInAsset(TSoftObjectPtr<USubtitleAssetUserData> SubtitleAsset);
 
 	// Queues an individual subtitle from a struct, or a single entry in a Subtitle asset's array.
+	UE_DEPRECATED(5.8, "Use QueueSingleSubtitle. It removes extraneous and confusing parameters.")
 	UFUNCTION(BlueprintCallable, Category = Subtitles)
 	static UE_API void QueueSubtitle(const FSubtitleAssetData& Subtitle, const ESubtitleTiming Timing = ESubtitleTiming::InternallyTimed);
 
+	// Queues an individual subtitle.	
+	// Parameters form an FSubtitleAssetData (see SubtitlesAndClosedCaptionsDelegates.h), but the Duration Type is removed as it is irrelevant for queueing this way.
+	UFUNCTION(BlueprintCallable, Category = Subtitles)
+	static UE_API void QueueSingleSubtitle(
+		const FText Text, const float Duration, const float StartOffset, const float Priority, const ESubtitleType SubtitleType, const ESubtitleTiming Timing = ESubtitleTiming::InternallyTimed);
+
+
 	UFUNCTION(BlueprintCallable, Category = Subtitles)
 	static UE_API bool IsSubtitleActive(const FSubtitleAssetData& Subtitle);
 
diff --git a/Engine/Source/Runtime/Engine/Public/Subtitles/SubtitlesAndClosedCaptionsDelegates.h b/Engine/Source/Runtime/Engine/Public/Subtitles/SubtitlesAndClosedCaptionsDelegates.h
--- a/Engine/Source/Runtime/Engine/Public/Subtitles/SubtitlesAndClosedCaptionsDelegates.h
+++ b/Engine/Source/Runtime/Engine/Public/Subtitles/SubtitlesAndClosedCaptionsDelegates.h
@@ -62,7 +62,7 @@ struct FSubtitleAssetData
 
 #if WITH_EDITORONLY_DATA
 	// Editor-only helper value for SubtitleDefaultDuration's EditCondition.
-	UPROPERTY(Transient)
+	UPROPERTY(Transient, VisibleDefaultsOnly, Category = "Subtitles|Duration")
 	bool bCanEditDuration = (SubtitleDurationType == ESubtitleDurationType::UseDurationProperty);
 #endif
 
