Subtitles in the Subtitles & Closed Captions plugin no longer squelch new subtitles equal in priority to what's currently displayed.

#virtualized
#Synced-CL 49422952

diff --git a/Engine/Plugins/Experimental/SubtitlesAndClosedCaptions/Source/SubtitlesAndClosedCaptions/Private/SubtitlesSubsystem.cpp b/Engine/Plugins/Experimental/SubtitlesAndClosedCaptions/Source/SubtitlesAndClosedCaptions/Private/SubtitlesSubsystem.cpp
--- a/Engine/Plugins/Experimental/SubtitlesAndClosedCaptions/Source/SubtitlesAndClosedCaptions/Private/SubtitlesSubsystem.cpp
+++ b/Engine/Plugins/Experimental/SubtitlesAndClosedCaptions/Source/SubtitlesAndClosedCaptions/Private/SubtitlesSubsystem.cpp
@@ -160,10 +160,11 @@ void USubtitlesSubsystem::AddAndDisplaySubtitle(FActiveSubtitleAssetData& NewAct
 	for (const FActiveSubtitleAssetData& PreQueuedSubtitle : ActiveSubtitles)
 	{
 		const FSubtitleAssetData& PreQueuedSubtitleData = PreQueuedSubtitle.Subtitle;
-		if (PreQueuedSubtitleData.SubtitleType == NewSubtitle.SubtitleType && (PreQueuedSubtitleData.Priority >= NewSubtitle.Priority))
+		if (PreQueuedSubtitleData.SubtitleType == NewSubtitle.SubtitleType && (PreQueuedSubtitleData.Priority > NewSubtitle.Priority))
 		{
-			UE_LOG(LogSubtitlesAndClosedCaptions, Display, TEXT("Subtitle %s won't display: its priority is lower than or equal to the currently-playing subtitle in that category."), *NewSubtitle.Text.ToString());
+			UE_LOG(LogSubtitlesAndClosedCaptions, Display, TEXT("Subtitle %s won't display: its priority is lower than the currently-playing subtitle in that category."), *NewSubtitle.Text.ToString());
 			ShouldAddSubtitle = false;
+			break;
 		}
 	}
 
