Enables SubtitlesAndClosedCaptions plugins across all platforms.

#Synced-CL 49740260

diff --git a/Engine/Plugins/Experimental/SubtitlesAndClosedCaptions/SubtitlesAndClosedCaptions.uplugin b/Engine/Plugins/Experimental/SubtitlesAndClosedCaptions/SubtitlesAndClosedCaptions.uplugin
--- a/Engine/Plugins/Experimental/SubtitlesAndClosedCaptions/SubtitlesAndClosedCaptions.uplugin
+++ b/Engine/Plugins/Experimental/SubtitlesAndClosedCaptions/SubtitlesAndClosedCaptions.uplugin
@@ -19,10 +19,7 @@
 		{
 			"Name": "SubtitlesAndClosedCaptions",
 			"Type": "Runtime",
-			"LoadingPhase": "Default",
-			"PlatformAllowList": [
-				"Win64"
-			]
+			"LoadingPhase": "Default"
 		},
 		{
 			"Name": "SubtitlesAndClosedCaptionsEditor",
