Stops subtitles recursively nested as AssetUserData in DialogueWaves/SoundCues when their parent Sound is stopped.

#Synced-CL 50264045

diff --git a/Engine/Plugins/Experimental/SubtitlesAndClosedCaptions/Source/SubtitlesAndClosedCaptions/Private/SubtitlesAudioSubsystem.cpp b/Engine/Plugins/Experimental/SubtitlesAndClosedCaptions/Source/SubtitlesAndClosedCaptions/Private/SubtitlesAudioSubsystem.cpp
--- a/Engine/Plugins/Experimental/SubtitlesAndClosedCaptions/Source/SubtitlesAndClosedCaptions/Private/SubtitlesAudioSubsystem.cpp
+++ b/Engine/Plugins/Experimental/SubtitlesAndClosedCaptions/Source/SubtitlesAndClosedCaptions/Private/SubtitlesAudioSubsystem.cpp
@@ -19,8 +19,8 @@ void USubtitlesAudioSubsystem::NotifyActiveSoundCreated(FActiveSound& ActiveSoun
 	{
 		if (Sound->HasAssetUserDataOfClass(USubtitleAssetUserData::StaticClass()))
 		{
-			const USubtitleAssetUserData* Subtitles = CastChecked<USubtitleAssetUserData>(Sound->GetAssetUserDataOfClass(USubtitleAssetUserData::StaticClass()));
-			for (const FSubtitleAssetData& Subtitle : Subtitles->Subtitles)
+			const USubtitleAssetUserData* SubtitleAsset = CastChecked<USubtitleAssetUserData>(Sound->GetAssetUserDataOfClass(USubtitleAssetUserData::StaticClass()));
+			for (const FSubtitleAssetData& Subtitle : SubtitleAsset->Subtitles)
 			{
 				const bool bUseDurationProperty = (Subtitle.SubtitleDurationType == ESubtitleDurationType::UseDurationProperty);
 				FSubtitlesAndClosedCaptionsDelegates::QueueSubtitle.ExecuteIfBound(FQueueSubtitleParameters{ Subtitle }, bUseDurationProperty ? ESubtitleTiming::InternallyTimed : ESubtitleTiming::ExternallyTimed);
@@ -45,18 +45,28 @@ void USubtitlesAudioSubsystem::NotifyActiveSoundCreated(FActiveSound& ActiveSoun
 			{
 				// Check inside Sound Cues for DialogueWaves as well.
 				const USoundCue* SoundCue = Cast<USoundCue>(Sound);
-				if (IsValid(SoundCue))
+
+				// SoundCue's FirstNode is used to grab all the other nodes, so it needs to be valid also.
+				if (IsValid(SoundCue) && IsValid(SoundCue->FirstNode))
 				{
-					const bool bQueuedFromTopLevelNode = CheckSoundCueNodeForDialogueWaves(SoundCue->FirstNode, ActiveSound);
+					const UDialogueSoundWaveProxy* Proxy = CheckSoundCueNodeForDialogueWaves(SoundCue->FirstNode, ActiveSound);
 
-					// If the parent node is not a DialogueWavePlayer, check for nested nodes that may also have DialogueWaves
-					if(!bQueuedFromTopLevelNode)
+					if (IsValid(Proxy))
+					{
+						QueueSubtitleFromDialogueWaveProxy(Proxy, ActiveSound);
+					}
+					else
 					{
+						// If the parent node is not a DialogueWavePlayer, check for nested nodes that may also have DialogueWaves
 						TArray<USoundNode*> ChildNodes;
 						SoundCue->FirstNode->GetAllNodes(ChildNodes);
 						for (const USoundNode* Node : ChildNodes)
 						{
-							CheckSoundCueNodeForDialogueWaves(Node, ActiveSound);
+							const UDialogueSoundWaveProxy* ChildProxy = CheckSoundCueNodeForDialogueWaves(Node, ActiveSound);
+							if (IsValid(ChildProxy))
+							{
+								QueueSubtitleFromDialogueWaveProxy(ChildProxy, ActiveSound);
+							}
 						}
 					}
 				}
@@ -73,19 +83,61 @@ void USubtitlesAudioSubsystem::NotifyActiveSoundDeleting(const FActiveSound& Act
 	{
 		if (Sound->HasAssetUserDataOfClass(USubtitleAssetUserData::StaticClass()))
 		{
-			if (Sound->HasAssetUserDataOfClass(USubtitleAssetUserData::StaticClass()))
+			const USubtitleAssetUserData* SubtitleAsset = CastChecked<USubtitleAssetUserData>(Sound->GetAssetUserDataOfClass(USubtitleAssetUserData::StaticClass()));
+			for (const FSubtitleAssetData& Subtitle : SubtitleAsset->Subtitles)
 			{
-				const USubtitleAssetUserData* Subtitles = CastChecked<USubtitleAssetUserData>(Sound->GetAssetUserDataOfClass(USubtitleAssetUserData::StaticClass()));
-				for (const FSubtitleAssetData& Subtitle : Subtitles->Subtitles)
+				const bool bUseDurationProperty = (Subtitle.SubtitleDurationType == ESubtitleDurationType::UseDurationProperty);
+				// If not using sound duration, subtitles continue to play instead of stopping:
+				// Their duration is completely decoupled from the sound's.
+
+				// ...but if Using the sound's duration: Stop All subtitles.
+				if (!bUseDurationProperty)
 				{
-					const bool bUseDurationProperty = (Subtitle.SubtitleDurationType == ESubtitleDurationType::UseDurationProperty);
-					// If not using sound duration, subtitles continue to play instead of stopping:
-					// Their duration is completely decoupled from the sound's.
+					FSubtitlesAndClosedCaptionsDelegates::StopSubtitle.ExecuteIfBound(Subtitle);
+				}
+			}
+		}
+
+		// Also check for DialogueWaves with localized text
+		static const IConsoleVariable* CVarUseNewSubtitles = IConsoleManager::Get().FindConsoleVariable(TEXT("au.UseNewSubtitles"));
 
-					// ...but if Using the sound's duration: Stop All subtitles.
-					if (!bUseDurationProperty)
+		if (CVarUseNewSubtitles && CVarUseNewSubtitles->GetBool() && ActiveSound.bHandleSubtitles)
+		{
+			FSubtitleAssetData SubtitleCue;
+
+			const UDialogueSoundWaveProxy* DialogueWave = Cast<UDialogueSoundWaveProxy>(Sound);
+			if (IsValid(DialogueWave))
+			{
+				StopSubtitleFromDialogueWaveProxy(DialogueWave, ActiveSound);
+			}
+			else
+			{
+				// Check inside Sound Cues for DialogueWaves as well.
+				const USoundCue* SoundCue = Cast<USoundCue>(Sound);
+
+				// SoundCue's FirstNode is used to grab all the other nodes, so it needs to be valid also.
+				if (IsValid(SoundCue) && IsValid(SoundCue->FirstNode))
+				{
+					const UDialogueSoundWaveProxy* Proxy = CheckSoundCueNodeForDialogueWaves(SoundCue->FirstNode, ActiveSound);
+
+					if (IsValid(Proxy))
+					{
+						// A dialogueWave won't reference another dialoguewave, so no need to search recursively
+						StopSubtitleFromDialogueWaveProxy(Proxy, ActiveSound);
+					}
+					else
 					{
-						FSubtitlesAndClosedCaptionsDelegates::StopSubtitle.ExecuteIfBound(Subtitle);
+						// If the parent node is not a DialogueWavePlayer, check for nested nodes that may also have DialogueWaves
+						TArray<USoundNode*> ChildNodes;
+						SoundCue->FirstNode->GetAllNodes(ChildNodes);
+						for (const USoundNode* Node : ChildNodes)
+						{
+							const UDialogueSoundWaveProxy* ChildProxy = CheckSoundCueNodeForDialogueWaves(Node, ActiveSound);
+							if (IsValid(ChildProxy))
+							{
+								StopSubtitleFromDialogueWaveProxy(ChildProxy, ActiveSound);
+							}
+						}
 					}
 				}
 			}
@@ -93,7 +145,7 @@ void USubtitlesAudioSubsystem::NotifyActiveSoundDeleting(const FActiveSound& Act
 	}
 }
 
-bool USubtitlesAudioSubsystem::CheckSoundCueNodeForDialogueWaves(const USoundNode* Node, const FActiveSound& ActiveSound)
+const UDialogueSoundWaveProxy* USubtitlesAudioSubsystem::CheckSoundCueNodeForDialogueWaves(const USoundNode* Node, const FActiveSound& ActiveSound)
 {
 	const USoundNodeDialoguePlayer* DialoguePlayerNode = Cast<USoundNodeDialoguePlayer>(Node);
 
@@ -104,12 +156,11 @@ bool USubtitlesAudioSubsystem::CheckSoundCueNodeForDialogueWaves(const USoundNod
 		const UDialogueSoundWaveProxy* Proxy = DialoguePlayerNode->GetContextualDialogueSoundWaveProxy();
 		if (IsValid(Proxy))
 		{
-			QueueSubtitleFromDialogueWaveProxy(Proxy, ActiveSound);
-			return true;
+			return Proxy;
 		}
 	}
 
-	return false;
+	return nullptr;
 }
 
 void USubtitlesAudioSubsystem::QueueSubtitleFromDialogueWaveProxy(const UDialogueSoundWaveProxy* DialogueWaveProxy, const FActiveSound& ActiveSound)
@@ -123,3 +174,15 @@ void USubtitlesAudioSubsystem::QueueSubtitleFromDialogueWaveProxy(const UDialogu
 
 	FSubtitlesAndClosedCaptionsDelegates::QueueSubtitle.ExecuteIfBound(FQueueSubtitleParameters{ SubtitleCue }, ESubtitleTiming::InternallyTimed);
 }
+
+void USubtitlesAudioSubsystem::StopSubtitleFromDialogueWaveProxy(const UDialogueSoundWaveProxy* DialogueWaveProxy, const FActiveSound& ActiveSound)
+{
+	check(ActiveSound.bHandleSubtitles);	// Function should only be called if the Active Sound actually has subtitles it needs queued.
+	FSubtitleAssetData SubtitleCue;
+	SubtitleCue.Text = DialogueWaveProxy->GetSubtitleText();
+	SubtitleCue.Priority = ActiveSound.SubtitlePriority;
+	SubtitleCue.Duration = DialogueWaveProxy->GetDuration();
+	SubtitleCue.StartOffset = ActiveSound.RequestedStartTime;
+
+	FSubtitlesAndClosedCaptionsDelegates::StopSubtitle.ExecuteIfBound(SubtitleCue);
+}
diff --git a/Engine/Plugins/Experimental/SubtitlesAndClosedCaptions/Source/SubtitlesAndClosedCaptions/Public/SubtitlesAudioSubsystem.h b/Engine/Plugins/Experimental/SubtitlesAndClosedCaptions/Source/SubtitlesAndClosedCaptions/Public/SubtitlesAudioSubsystem.h
--- a/Engine/Plugins/Experimental/SubtitlesAndClosedCaptions/Source/SubtitlesAndClosedCaptions/Public/SubtitlesAudioSubsystem.h
+++ b/Engine/Plugins/Experimental/SubtitlesAndClosedCaptions/Source/SubtitlesAndClosedCaptions/Public/SubtitlesAudioSubsystem.h
@@ -31,8 +31,10 @@ public:
 
 private:
 	// Helper functions to cut down on duplicated code
-	static bool CheckSoundCueNodeForDialogueWaves(const USoundNode* Node, const FActiveSound& ActiveSound);
+	static const UDialogueSoundWaveProxy* CheckSoundCueNodeForDialogueWaves(const USoundNode* Node, const FActiveSound& ActiveSound);
 	static void QueueSubtitleFromDialogueWaveProxy(const UDialogueSoundWaveProxy* DialogueWaveProxy, const FActiveSound& ActiveSound);
+
+	static void StopSubtitleFromDialogueWaveProxy(const UDialogueSoundWaveProxy* DialogueWaveProxy, const FActiveSound& ActiveSound);
 };
 
 #undef UE_API
