Remove old lower-priority subtitles when a new higher-pri subtitle is queued. Also ensure stop calls can't break the UI widget.

#Synced-CL 50417928

diff --git a/Engine/Plugins/Experimental/SubtitlesAndClosedCaptions/Source/SubtitlesAndClosedCaptions/Private/SubtitleWidget.cpp b/Engine/Plugins/Experimental/SubtitlesAndClosedCaptions/Source/SubtitlesAndClosedCaptions/Private/SubtitleWidget.cpp
--- a/Engine/Plugins/Experimental/SubtitlesAndClosedCaptions/Source/SubtitlesAndClosedCaptions/Private/SubtitleWidget.cpp
+++ b/Engine/Plugins/Experimental/SubtitlesAndClosedCaptions/Source/SubtitlesAndClosedCaptions/Private/SubtitleWidget.cpp
@@ -11,35 +11,11 @@ void USubtitleWidget::NativeConstruct()
 {
 	Super::NativeConstruct();
 
-	// Ensure the text blocks don't start visible.
-	if (IsValid(DialogSubtitleBlock))
-	{
-		DialogSubtitleBlock->SetVisibility(ESlateVisibility::Collapsed);
-	}
-
-	if (IsValid(CaptionSubtitleBlock))
-	{
-		CaptionSubtitleBlock->SetVisibility(ESlateVisibility::Collapsed);
-	}
-
-	if (IsValid(DescriptionSubtitleBlock))
-	{
-		DescriptionSubtitleBlock->SetVisibility(ESlateVisibility::Collapsed);
-	}
+	CurrentlyDisplayedSubtitles.Add({ ESubtitleType::Subtitle,			TOptional<FSubtitleAssetData>() });
+	CurrentlyDisplayedSubtitles.Add({ ESubtitleType::AudioDescription,	TOptional<FSubtitleAssetData>() });
+	CurrentlyDisplayedSubtitles.Add({ ESubtitleType::ClosedCaption,		TOptional<FSubtitleAssetData>() });
 
-	// Borders should also not be visible on startup.
-	if (IsValid(DialogBorder))
-	{
-		DialogBorder->SetVisibility(ESlateVisibility::Collapsed);
-	}
-	if (IsValid(CaptionBorder))
-	{
-		CaptionBorder->SetVisibility(ESlateVisibility::Collapsed);
-	}
-	if (IsValid(DescriptionBorder))
-	{
-		DescriptionBorder->SetVisibility(ESlateVisibility::Collapsed);
-	}
+	StopDisplayingAllSubtitles();
 }
 
 void USubtitleWidget::StartDisplayingSubtitle(const FSubtitleAssetData& Subtitle)
@@ -48,7 +24,9 @@ void USubtitleWidget::StartDisplayingSubtitle(const FSubtitleAssetData& Subtitle
 	UTextBlock* TextToModify = nullptr;
 	UBorder* BorderToModify = nullptr;
 
-	switch (Subtitle.SubtitleType)
+	const ESubtitleType SubtitleType = Subtitle.SubtitleType;
+
+	switch (SubtitleType)
 	{
 	default:
 	{
@@ -90,47 +68,88 @@ void USubtitleWidget::StartDisplayingSubtitle(const FSubtitleAssetData& Subtitle
 		{
 			BorderToModify->SetVisibility(ESlateVisibility::HitTestInvisible);
 		}
+
+		// Track the currently-displayed subtitle for comparison when removing later
+		CurrentlyDisplayedSubtitles[SubtitleType] = Subtitle;
 	}
 }
 
-void USubtitleWidget::StopDisplayingSubtitle(const ESubtitleType SubtitleType)
+void USubtitleWidget::StopDisplayingSubtitle(const FSubtitleAssetData& Subtitle)
 {
 	UTextBlock* TextToModify = nullptr;
 	UBorder* BorderToModify = nullptr;
 
-	switch (SubtitleType)
+	// Confirm that the removed subtitle is actually being displayed before trying to stop it.
+	const ESubtitleType SubtitleType = Subtitle.SubtitleType;
+	if (CurrentlyDisplayedSubtitles[SubtitleType].IsSet() && CurrentlyDisplayedSubtitles[SubtitleType] == Subtitle)
 	{
-	default: // Fallthrough, no break (see ::StartDisplayingSubtitle)
-	case ESubtitleType::Subtitle:
-	{
-		TextToModify = DialogSubtitleBlock;
-		BorderToModify = DialogBorder;
-		break;
+		switch (SubtitleType)
+		{
+		default: // Fallthrough, no break (see ::StartDisplayingSubtitle)
+		case ESubtitleType::Subtitle:
+		{
+			TextToModify = DialogSubtitleBlock;
+			BorderToModify = DialogBorder;
+			break;
+		}
+
+		case ESubtitleType::ClosedCaption:
+		{
+			TextToModify = CaptionSubtitleBlock;
+			BorderToModify = CaptionBorder;
+			break;
+		}
+
+		case ESubtitleType::AudioDescription:
+		{
+			TextToModify = DescriptionSubtitleBlock;
+			BorderToModify = DescriptionBorder;
+			break;
+		}
+
+		}
+
+		if (IsValid(TextToModify))
+		{
+			TextToModify->SetVisibility(ESlateVisibility::Collapsed);
+
+			if (IsValid(BorderToModify))
+			{
+				BorderToModify->SetVisibility(ESlateVisibility::Collapsed);
+			}
+		}
 	}
+}
 
-	case ESubtitleType::ClosedCaption:
+void USubtitleWidget::StopDisplayingAllSubtitles()
+{
+	// Ensure the text blocks don't start visible.
+	if (IsValid(DialogSubtitleBlock))
 	{
-		TextToModify = CaptionSubtitleBlock;
-		BorderToModify = CaptionBorder;
-		break;
+		DialogSubtitleBlock->SetVisibility(ESlateVisibility::Collapsed);
 	}
 
-	case ESubtitleType::AudioDescription:
+	if (IsValid(CaptionSubtitleBlock))
 	{
-		TextToModify = DescriptionSubtitleBlock;
-		BorderToModify = DescriptionBorder;
-		break;
+		CaptionSubtitleBlock->SetVisibility(ESlateVisibility::Collapsed);
 	}
 
+	if (IsValid(DescriptionSubtitleBlock))
+	{
+		DescriptionSubtitleBlock->SetVisibility(ESlateVisibility::Collapsed);
 	}
 
-	if (IsValid(TextToModify))
+	// Borders should also not be visible on startup.
+	if (IsValid(DialogBorder))
 	{
-		TextToModify->SetVisibility(ESlateVisibility::Collapsed);
-
-		if (IsValid(BorderToModify))
-		{
-			BorderToModify->SetVisibility(ESlateVisibility::Collapsed);
-		}
+		DialogBorder->SetVisibility(ESlateVisibility::Collapsed);
+	}
+	if (IsValid(CaptionBorder))
+	{
+		CaptionBorder->SetVisibility(ESlateVisibility::Collapsed);
+	}
+	if (IsValid(DescriptionBorder))
+	{
+		DescriptionBorder->SetVisibility(ESlateVisibility::Collapsed);
 	}
 }
diff --git a/Engine/Plugins/Experimental/SubtitlesAndClosedCaptions/Source/SubtitlesAndClosedCaptions/Private/SubtitlesSubsystem.cpp b/Engine/Plugins/Experimental/SubtitlesAndClosedCaptions/Source/SubtitlesAndClosedCaptions/Private/SubtitlesSubsystem.cpp
--- a/Engine/Plugins/Experimental/SubtitlesAndClosedCaptions/Source/SubtitlesAndClosedCaptions/Private/SubtitlesSubsystem.cpp
+++ b/Engine/Plugins/Experimental/SubtitlesAndClosedCaptions/Source/SubtitlesAndClosedCaptions/Private/SubtitlesSubsystem.cpp
@@ -58,7 +58,6 @@ void USubtitlesSubsystem::QueueSubtitle(FQueueSubtitleParameters Params, const E
 	const FSubtitleAssetData& Subtitle = Params.Subtitle;
 	UE_LOG(LogSubtitlesAndClosedCaptions, Display, TEXT("QueueSubtitle: '%s'"), *Subtitle.Text.ToString());
 
-
 	// Externally-timed subtitles will be removed by the system queueing them, so they have an otherwise-infinite duration.
 	float Duration = InfiniteDuration;
 	if (Timing == ESubtitleTiming::InternallyTimed)
@@ -155,13 +154,15 @@ void USubtitlesSubsystem::MakeDelayedSubtitleActive(const FSubtitleAssetData& Su
 void USubtitlesSubsystem::AddAndDisplaySubtitle(FActiveSubtitleAssetData& NewActiveSubtitle)
 {
 	const FSubtitleAssetData& NewSubtitle = NewActiveSubtitle.Subtitle;
+	const float NewPriority = NewSubtitle.Priority;
+	const ESubtitleType Category = NewSubtitle.SubtitleType;
 
 	bool ShouldAddSubtitle = true;
 	// Warn if the subtitle won't play due to being lower-priority than the currently-playing subtitle in that category.
 	for (const FActiveSubtitleAssetData& PreQueuedSubtitle : ActiveSubtitles)
 	{
 		const FSubtitleAssetData& PreQueuedSubtitleData = PreQueuedSubtitle.Subtitle;
-		if (PreQueuedSubtitleData.SubtitleType == NewSubtitle.SubtitleType && (PreQueuedSubtitleData.Priority > NewSubtitle.Priority))
+		if (PreQueuedSubtitleData.SubtitleType == Category && (PreQueuedSubtitleData.Priority > NewPriority))
 		{
 			UE_LOG(LogSubtitlesAndClosedCaptions, Display, TEXT("Subtitle %s won't display: its priority is lower than the currently-playing subtitle in that category."), *NewSubtitle.Text.ToString());
 			ShouldAddSubtitle = false;
@@ -169,6 +170,19 @@ void USubtitlesSubsystem::AddAndDisplaySubtitle(FActiveSubtitleAssetData& NewAct
 		}
 	}
 
+	// Remove newly lower-priority subtitles.
+	FTimerManager& TimerManager = GetWorldRef().GetTimerManager();
+
+	const int32 FirstRemovedIndex = Algo::StableRemoveIf(ActiveSubtitles,
+		[NewPriority, Category](const FActiveSubtitleAssetData& RemovalCandidate) {return (RemovalCandidate.Subtitle.Priority < NewPriority) && RemovalCandidate.Subtitle.SubtitleType == Category; }
+	);
+	for (int32 Index = FirstRemovedIndex; Index < ActiveSubtitles.Num(); ++Index)
+	{
+		TimerManager.ClearTimer(ActiveSubtitles[Index].DurationTimerHandle);
+	}
+
+	ActiveSubtitles.SetNum(FirstRemovedIndex);
+
 	if (ShouldAddSubtitle)
 	{
 		ActiveSubtitles.Add(MoveTemp(NewActiveSubtitle));
@@ -219,9 +233,7 @@ void USubtitlesSubsystem::StopAllSubtitles()
 	// Clear the widget's display
 	if (IsValid(SubtitleWidget))
 	{
-		SubtitleWidget->StopDisplayingSubtitle(ESubtitleType::AudioDescription);
-		SubtitleWidget->StopDisplayingSubtitle(ESubtitleType::ClosedCaption);
-		SubtitleWidget->StopDisplayingSubtitle(ESubtitleType::Subtitle);
+		SubtitleWidget->StopDisplayingAllSubtitles();
 	}
 	else 
 	{
@@ -264,7 +276,7 @@ void USubtitlesSubsystem::RemoveActiveSubtitle(const FSubtitleAssetData& Subtitl
 	// Stop Displaying the removed subtitle and display a newly-most-relevant one if applicable.
 	if (bSuccessfullyRemoved && IsValid(SubtitleWidget))
 	{
-		SubtitleWidget->StopDisplayingSubtitle(Subtitle.SubtitleType);
+		SubtitleWidget->StopDisplayingSubtitle(Subtitle);
 
 		if (!ActiveSubtitles.IsEmpty())
 		{
diff --git a/Engine/Plugins/Experimental/SubtitlesAndClosedCaptions/Source/SubtitlesAndClosedCaptions/Public/SubtitleWidget.h b/Engine/Plugins/Experimental/SubtitlesAndClosedCaptions/Source/SubtitlesAndClosedCaptions/Public/SubtitleWidget.h
--- a/Engine/Plugins/Experimental/SubtitlesAndClosedCaptions/Source/SubtitlesAndClosedCaptions/Public/SubtitleWidget.h
+++ b/Engine/Plugins/Experimental/SubtitlesAndClosedCaptions/Source/SubtitlesAndClosedCaptions/Public/SubtitleWidget.h
@@ -16,6 +16,9 @@ class USubtitleWidget : public UUserWidget
 {
 	GENERATED_BODY()
 
+private:
+	TMap<ESubtitleType, TOptional<FSubtitleAssetData>> CurrentlyDisplayedSubtitles;
+
 protected:
 	UE_API virtual void NativeConstruct() override;
 
@@ -40,7 +43,8 @@ protected:
 	TObjectPtr<UBorder> DescriptionBorder;
 public:
 	void StartDisplayingSubtitle(const FSubtitleAssetData& Subtitle);
-	void StopDisplayingSubtitle(const ESubtitleType SubtitleType);
+	void StopDisplayingSubtitle(const FSubtitleAssetData& Subtitle);
+	void StopDisplayingAllSubtitles();
 };
 
 #undef UE_API
\ No newline at end of file
