[Subtitles] To prevent crashes from incomplete initialization, set up the widget's currently-displayed map in a standard c++ constructor, instead of the inherited NativeConstruct function.

#Synced-CL 50662405

diff --git a/Engine/Plugins/Experimental/SubtitlesAndClosedCaptions/Source/SubtitlesAndClosedCaptions/Private/SubtitleWidget.cpp b/Engine/Plugins/Experimental/SubtitlesAndClosedCaptions/Source/SubtitlesAndClosedCaptions/Private/SubtitleWidget.cpp
--- a/Engine/Plugins/Experimental/SubtitlesAndClosedCaptions/Source/SubtitlesAndClosedCaptions/Private/SubtitleWidget.cpp
+++ b/Engine/Plugins/Experimental/SubtitlesAndClosedCaptions/Source/SubtitlesAndClosedCaptions/Private/SubtitleWidget.cpp
@@ -11,10 +11,6 @@ void USubtitleWidget::NativeConstruct()
 {
 	Super::NativeConstruct();
 
-	CurrentlyDisplayedSubtitles.Add({ ESubtitleType::Subtitle,			TOptional<FSubtitleAssetData>() });
-	CurrentlyDisplayedSubtitles.Add({ ESubtitleType::AudioDescription,	TOptional<FSubtitleAssetData>() });
-	CurrentlyDisplayedSubtitles.Add({ ESubtitleType::ClosedCaption,		TOptional<FSubtitleAssetData>() });
-
 	StopDisplayingAllSubtitles();
 }
 
@@ -153,3 +149,12 @@ void USubtitleWidget::StopDisplayingAllSubtitles()
 		DescriptionBorder->SetVisibility(ESlateVisibility::Collapsed);
 	}
 }
+
+USubtitleWidget::USubtitleWidget(const FObjectInitializer& ObjectInitializer)
+	: UUserWidget(ObjectInitializer)
+{
+
+	CurrentlyDisplayedSubtitles.Add({ ESubtitleType::Subtitle,			TOptional<FSubtitleAssetData>() });
+	CurrentlyDisplayedSubtitles.Add({ ESubtitleType::AudioDescription,	TOptional<FSubtitleAssetData>() });
+	CurrentlyDisplayedSubtitles.Add({ ESubtitleType::ClosedCaption,		TOptional<FSubtitleAssetData>() });
+}
diff --git a/Engine/Plugins/Experimental/SubtitlesAndClosedCaptions/Source/SubtitlesAndClosedCaptions/Public/SubtitleWidget.h b/Engine/Plugins/Experimental/SubtitlesAndClosedCaptions/Source/SubtitlesAndClosedCaptions/Public/SubtitleWidget.h
--- a/Engine/Plugins/Experimental/SubtitlesAndClosedCaptions/Source/SubtitlesAndClosedCaptions/Public/SubtitleWidget.h
+++ b/Engine/Plugins/Experimental/SubtitlesAndClosedCaptions/Source/SubtitlesAndClosedCaptions/Public/SubtitleWidget.h
@@ -45,6 +45,8 @@ public:
 	void StartDisplayingSubtitle(const FSubtitleAssetData& Subtitle);
 	void StopDisplayingSubtitle(const FSubtitleAssetData& Subtitle);
 	void StopDisplayingAllSubtitles();
+
+	UE_API USubtitleWidget(const FObjectInitializer& ObjectInitializer);
 };
 
 #undef UE_API
\ No newline at end of file
