Subtitles: Change UE_LOG calls to UE_LOGFMT for cleaner code.

#Synced-CL 51411063

diff --git a/Engine/Plugins/Experimental/SubtitlesAndClosedCaptions/Source/SubtitlesAndClosedCaptions/Private/SubtitleWidget.cpp b/Engine/Plugins/Experimental/SubtitlesAndClosedCaptions/Source/SubtitlesAndClosedCaptions/Private/SubtitleWidget.cpp
--- a/Engine/Plugins/Experimental/SubtitlesAndClosedCaptions/Source/SubtitlesAndClosedCaptions/Private/SubtitleWidget.cpp
+++ b/Engine/Plugins/Experimental/SubtitlesAndClosedCaptions/Source/SubtitlesAndClosedCaptions/Private/SubtitleWidget.cpp
@@ -3,6 +3,7 @@
 #include "SubtitleWidget.h"
 #include "Components/Border.h"
 #include "Components/TextBlock.h"
+#include "Logging/StructuredLog.h"
 #include "SubtitlesAndClosedCaptionsModule.h"
 
 #include UE_INLINE_GENERATED_CPP_BY_NAME(SubtitleWidget)
@@ -26,7 +27,7 @@ void USubtitleWidget::StartDisplayingSubtitle(const FSubtitleAssetData& Subtitle
 	{
 	default:
 	{
-		UE_LOG(LogSubtitlesAndClosedCaptions, Warning, TEXT("An unrecognized subtitle type was queued for display. Using the standard Subtitle text block as a fallback."));
+		UE_LOGFMT(LogSubtitlesAndClosedCaptions, Warning, "An unrecognized subtitle type was queued for display. Using the standard Subtitle text block as a fallback.");
 		// Fallthrough, no break.
 	}
 
diff --git a/Engine/Plugins/Experimental/SubtitlesAndClosedCaptions/Source/SubtitlesAndClosedCaptions/Private/SubtitlesSubsystem.cpp b/Engine/Plugins/Experimental/SubtitlesAndClosedCaptions/Source/SubtitlesAndClosedCaptions/Private/SubtitlesSubsystem.cpp
--- a/Engine/Plugins/Experimental/SubtitlesAndClosedCaptions/Source/SubtitlesAndClosedCaptions/Private/SubtitlesSubsystem.cpp
+++ b/Engine/Plugins/Experimental/SubtitlesAndClosedCaptions/Source/SubtitlesAndClosedCaptions/Private/SubtitlesSubsystem.cpp
@@ -10,6 +10,7 @@
 #include "Engine/Font.h"
 #include "Engine/World.h"
 #include "Engine/Engine.h"
+#include "Logging/StructuredLog.h"
 #include "Math/UnrealMathUtility.h"
 #include "Styling/CoreStyle.h"
 #include "SubtitlesAndClosedCaptionsModule.h"
@@ -47,7 +48,7 @@ void USubtitlesSubsystem::QueueSubtitle(FQueueSubtitleParameters Params, const E
 	// Don't queue subtitles when the engine doesn't have them enabled or forced off. 
 	if (GEngine == nullptr)
 	{
-		UE_LOG(LogSubtitlesAndClosedCaptions, Warning, TEXT("Couldn't check if subtitles are enabled, because GEngine nullptr."));
+		UE_LOGFMT(LogSubtitlesAndClosedCaptions, Warning, "Couldn't check if subtitles are enabled, because GEngine nullptr.");
 		return;
 	}
 	if (!GEngine->bSubtitlesEnabled || GEngine->bSubtitlesForcedOff)
@@ -56,7 +57,7 @@ void USubtitlesSubsystem::QueueSubtitle(FQueueSubtitleParameters Params, const E
 	}
 
 	const FSubtitleAssetData& Subtitle = Params.Subtitle;
-	UE_LOG(LogSubtitlesAndClosedCaptions, Display, TEXT("QueueSubtitle: '%s'"), *Subtitle.Text.ToString());
+	UE_LOGFMT(LogSubtitlesAndClosedCaptions, Display, "QueueSubtitle: {Name}", Subtitle.Text.ToString());
 
 	// Externally-timed subtitles will be removed by the system queueing them, so they have an otherwise-infinite duration.
 	float Duration = InfiniteDuration;
@@ -164,7 +165,7 @@ void USubtitlesSubsystem::AddAndDisplaySubtitle(FActiveSubtitleAssetData& NewAct
 		const FSubtitleAssetData& PreQueuedSubtitleData = PreQueuedSubtitle.Subtitle;
 		if (PreQueuedSubtitleData.SubtitleType == Category && (PreQueuedSubtitleData.Priority > NewPriority))
 		{
-			UE_LOG(LogSubtitlesAndClosedCaptions, Display, TEXT("Subtitle %s won't display: its priority is lower than the currently-playing subtitle in that category."), *NewSubtitle.Text.ToString());
+			UE_LOGFMT(LogSubtitlesAndClosedCaptions, Display, "Subtitle {Name} won't display: its priority is lower than the currently-playing subtitle in that category.", NewSubtitle.Text.ToString());
 			ShouldAddSubtitle = false;
 			break;
 		}
@@ -237,7 +238,7 @@ void USubtitlesSubsystem::StopAllSubtitles()
 	}
 	else 
 	{
-		UE_LOG(LogSubtitlesAndClosedCaptions, Warning, TEXT("Can't remove subtitles because there isn't a valid UMG widget."));
+		UE_LOGFMT(LogSubtitlesAndClosedCaptions, Warning, "Can't remove subtitles because there isn't a valid UMG widget.");
 	}
 }
 
@@ -245,13 +246,13 @@ void USubtitlesSubsystem::ReplaceWidget(const TSoftClassPtr<USubtitleWidget>& Ne
 {
 	if (!ActiveSubtitles.IsEmpty() || !DelayedSubtitles.IsEmpty())
 	{
-		UE_LOG(LogSubtitlesAndClosedCaptions, Error, TEXT("Can't replace the subtitle widget: there are still subtitles queued or displaying. Use StopAllSubtitles() first."));
+		UE_LOGFMT(LogSubtitlesAndClosedCaptions, Error, "Can't replace the subtitle widget: there are still subtitles queued or displaying. Use StopAllSubtitles() first.");
 		return;
 	}
 
 	if (!TryCreateUMGWidgetFromAsset(NewWidgetAsset))
 	{
-		UE_LOG(LogSubtitlesAndClosedCaptions, Error, TEXT("Can't replace the subtitle widget; was a valid asset provided?"));
+		UE_LOGFMT(LogSubtitlesAndClosedCaptions, Error, "Can't replace the subtitle widget; was a valid asset provided?");
 	}
 }
 
@@ -286,7 +287,7 @@ void USubtitlesSubsystem::RemoveActiveSubtitle(const FSubtitleAssetData& Subtitl
 	}
 	else if (!IsValid(SubtitleWidget))
 	{
-		UE_LOG(LogSubtitlesAndClosedCaptions, Warning, TEXT("Can't remove subtitles because there isn't a valid UMG widget."));
+		UE_LOGFMT(LogSubtitlesAndClosedCaptions, Warning, "Can't remove subtitles because there isn't a valid UMG widget.");
 	}
 
 	// Remove delayed-start subtitles that haven't been queued yet that also use this asset.
@@ -337,11 +338,11 @@ bool USubtitlesSubsystem::TryCreateUMGWidget()
 				SubtitleWidget->SetFlags(RF_Transient);
 			}
 			// Log the fallback as a warning as seeing the default widget probably isn't intended.
-			UE_LOG(LogSubtitlesAndClosedCaptions, Warning, TEXT("A valid Subtitle Widget wasn't provided. The plugin's default widget will be used as a fallback."));
+			UE_LOGFMT(LogSubtitlesAndClosedCaptions, Warning, "A valid Subtitle Widget wasn't provided. The plugin's default widget will be used as a fallback.");
 		}
 		else
 		{
-			UE_LOG(LogSubtitlesAndClosedCaptions, Error, TEXT("The default Subtitle Widget asset isn't valid. Subtitles won't be displayed."));
+			UE_LOGFMT(LogSubtitlesAndClosedCaptions, Error, "The default Subtitle Widget asset isn't valid. Subtitles won't be displayed.");
 		}
 	}
 
@@ -366,7 +367,7 @@ void USubtitlesSubsystem::UpdateWidgetData()
 	}
 	else
 	{
-		UE_LOG(LogSubtitlesAndClosedCaptions, Warning, TEXT("Can't display subtitles because there isn't a valid UMG widget to display it to (check Project Settings)."));
+		UE_LOGFMT(LogSubtitlesAndClosedCaptions, Warning, "Can't display subtitles because there isn't a valid UMG widget to display it to (check Project Settings).");
 	}
 }
 
diff --git a/Engine/Plugins/Experimental/SubtitlesAndClosedCaptions/Source/SubtitlesAndClosedCaptions/Public/MovieSceneSubtitlesSystem.h b/Engine/Plugins/Experimental/SubtitlesAndClosedCaptions/Source/SubtitlesAndClosedCaptions/Public/MovieSceneSubtitlesSystem.h
--- a/Engine/Plugins/Experimental/SubtitlesAndClosedCaptions/Source/SubtitlesAndClosedCaptions/Public/MovieSceneSubtitlesSystem.h
+++ b/Engine/Plugins/Experimental/SubtitlesAndClosedCaptions/Source/SubtitlesAndClosedCaptions/Public/MovieSceneSubtitlesSystem.h
@@ -71,10 +71,8 @@ struct FEvaluateSubtitles
 			/* TODO: enable for ClangEditor Win64
 			const FString LastStatusString = Enum->GetNameStringByIndex(LastStatus);
 			const FString StatusString = Enum->GetNameStringByIndex(Status);
-			UE_LOG(LogSubtitlesAndClosedCaptions, Verbose, TEXT("LastStatus: %s, Status: %s, HasJumped: %u"),
-				*LastStatusString,
-				*StatusString,
-				!!Context.HasJumped());
+			UE_LOGFMT(LogSubtitlesAndClosedCaptions, Verbose, "LastStatus: {LastStatus}, Status: {CurStatus}, HasJumped: {Jumped}",
+				LastStatusString, StatusString, Context.HasJumped());
 			*/
 		}
 #endif // !NO_LOGGING
diff --git a/Engine/Plugins/Experimental/SubtitlesAndClosedCaptions/Source/SubtitlesAndClosedCaptionsEditor/Private/SubtitleMigrationUtility.cpp b/Engine/Plugins/Experimental/SubtitlesAndClosedCaptions/Source/SubtitlesAndClosedCaptionsEditor/Private/SubtitleMigrationUtility.cpp
--- a/Engine/Plugins/Experimental/SubtitlesAndClosedCaptions/Source/SubtitlesAndClosedCaptionsEditor/Private/SubtitleMigrationUtility.cpp
+++ b/Engine/Plugins/Experimental/SubtitlesAndClosedCaptions/Source/SubtitlesAndClosedCaptionsEditor/Private/SubtitleMigrationUtility.cpp
@@ -4,6 +4,7 @@
 
 #include "EditorUtilityLibrary.h"
 #include "Engine/Engine.h"
+#include "Logging/StructuredLog.h"
 #include "Sound/SoundWave.h"
 #include "SubtitlesAndClosedCaptionsModule.h"
 #include "Subtitles/SubtitlesAndClosedCaptionsDelegates.h"
@@ -86,12 +87,12 @@ void USubtitleMigrationUtility::RemoveLegacySubtitles()
 				}
 				else
 				{
-					UE_LOG(LogSubtitlesAndClosedCaptions, Warning, TEXT("Can't remove legacy subtitles: migrate them to the new system first with Convert Legacy Subtitles."));
+					UE_LOGFMT(LogSubtitlesAndClosedCaptions, Warning, "Can't remove legacy subtitles: migrate them to the new system first with Convert Legacy Subtitles.");
 				}
 			}
 			else
 			{
-				UE_LOG(LogSubtitlesAndClosedCaptions, Warning, TEXT("Remove Legacy Subtitles can only be used on SoundWave assets."));
+				UE_LOGFMT(LogSubtitlesAndClosedCaptions, Warning, "Remove Legacy Subtitles can only be used on SoundWave assets.");
 			}
 		}
 
@@ -148,7 +149,7 @@ void USubtitleMigrationUtility::ConvertLegacySubtitles()
 			}
 			else
 			{
-				UE_LOG(LogSubtitlesAndClosedCaptions, Warning, TEXT("Convert Legacy Subtitles can only be used on SoundWave assets."));
+				UE_LOGFMT(LogSubtitlesAndClosedCaptions, Warning, "Convert Legacy Subtitles can only be used on SoundWave assets.");
 			}
 		}
 
