Move SubtitlesAndClosedCaptions and SubtitlesTests plugins out of experimental and restricted folders.
#virtualized
#Synced-CL 51514915
#Synced-CL 51514915

#Synced-CL 51514915

diff --git a/Engine/Plugins/Experimental/SubtitlesAndClosedCaptions/Source/SubtitlesAndClosedCaptions/Private/MovieSceneSubtitleComponentTypes.cpp b/Engine/Plugins/Experimental/SubtitlesAndClosedCaptions/Source/SubtitlesAndClosedCaptions/Private/MovieSceneSubtitleComponentTypes.cpp
deleted file mode 100644
diff --git a/Engine/Plugins/Experimental/SubtitlesAndClosedCaptions/Source/SubtitlesAndClosedCaptions/Private/MovieSceneSubtitleSection.cpp b/Engine/Plugins/Experimental/SubtitlesAndClosedCaptions/Source/SubtitlesAndClosedCaptions/Private/MovieSceneSubtitleSection.cpp
deleted file mode 100644
diff --git a/Engine/Plugins/Experimental/SubtitlesAndClosedCaptions/Source/SubtitlesAndClosedCaptions/Private/MovieSceneSubtitlesSystem.cpp b/Engine/Plugins/Experimental/SubtitlesAndClosedCaptions/Source/SubtitlesAndClosedCaptions/Private/MovieSceneSubtitlesSystem.cpp
deleted file mode 100644
diff --git a/Engine/Plugins/Experimental/SubtitlesAndClosedCaptions/Source/SubtitlesAndClosedCaptions/Private/MovieSceneSubtitlesTrack.cpp b/Engine/Plugins/Experimental/SubtitlesAndClosedCaptions/Source/SubtitlesAndClosedCaptions/Private/MovieSceneSubtitlesTrack.cpp
deleted file mode 100644
diff --git a/Engine/Plugins/Experimental/SubtitlesAndClosedCaptions/Source/SubtitlesAndClosedCaptions/Private/SubtitleWidget.cpp b/Engine/Plugins/Experimental/SubtitlesAndClosedCaptions/Source/SubtitlesAndClosedCaptions/Private/SubtitleWidget.cpp
deleted file mode 100644
diff --git a/Engine/Plugins/Experimental/SubtitlesAndClosedCaptions/Source/SubtitlesAndClosedCaptions/Private/SubtitlesAndClosedCaptionsModule.cpp b/Engine/Plugins/Experimental/SubtitlesAndClosedCaptions/Source/SubtitlesAndClosedCaptions/Private/SubtitlesAndClosedCaptionsModule.cpp
deleted file mode 100644
diff --git a/Engine/Plugins/Experimental/SubtitlesAndClosedCaptions/Source/SubtitlesAndClosedCaptions/Private/SubtitlesAudioSubsystem.cpp b/Engine/Plugins/Experimental/SubtitlesAndClosedCaptions/Source/SubtitlesAndClosedCaptions/Private/SubtitlesAudioSubsystem.cpp
deleted file mode 100644
diff --git a/Engine/Plugins/Experimental/SubtitlesAndClosedCaptions/Source/SubtitlesAndClosedCaptions/Private/SubtitlesBlueprintFunctionLibrary.cpp b/Engine/Plugins/Experimental/SubtitlesAndClosedCaptions/Source/SubtitlesAndClosedCaptions/Private/SubtitlesBlueprintFunctionLibrary.cpp
deleted file mode 100644
diff --git a/Engine/Plugins/Experimental/SubtitlesAndClosedCaptions/Source/SubtitlesAndClosedCaptions/Private/SubtitlesSettings.cpp b/Engine/Plugins/Experimental/SubtitlesAndClosedCaptions/Source/SubtitlesAndClosedCaptions/Private/SubtitlesSettings.cpp
deleted file mode 100644
diff --git a/Engine/Plugins/Experimental/SubtitlesAndClosedCaptions/Source/SubtitlesAndClosedCaptions/Private/SubtitlesSubsystem.cpp b/Engine/Plugins/Experimental/SubtitlesAndClosedCaptions/Source/SubtitlesAndClosedCaptions/Private/SubtitlesSubsystem.cpp
deleted file mode 100644
diff --git a/Engine/Plugins/Experimental/SubtitlesAndClosedCaptions/Source/SubtitlesAndClosedCaptions/Public/ActiveSubtitleAssetData.h b/Engine/Plugins/Experimental/SubtitlesAndClosedCaptions/Source/SubtitlesAndClosedCaptions/Public/ActiveSubtitleAssetData.h
deleted file mode 100644
diff --git a/Engine/Plugins/Experimental/SubtitlesAndClosedCaptions/Source/SubtitlesAndClosedCaptions/Public/MovieSceneSubtitleComponentTypes.h b/Engine/Plugins/Experimental/SubtitlesAndClosedCaptions/Source/SubtitlesAndClosedCaptions/Public/MovieSceneSubtitleComponentTypes.h
deleted file mode 100644
diff --git a/Engine/Plugins/Experimental/SubtitlesAndClosedCaptions/Source/SubtitlesAndClosedCaptions/Public/MovieSceneSubtitleSection.h b/Engine/Plugins/Experimental/SubtitlesAndClosedCaptions/Source/SubtitlesAndClosedCaptions/Public/MovieSceneSubtitleSection.h
deleted file mode 100644
diff --git a/Engine/Plugins/Experimental/SubtitlesAndClosedCaptions/Source/SubtitlesAndClosedCaptions/Public/MovieSceneSubtitlesSystem.h b/Engine/Plugins/Experimental/SubtitlesAndClosedCaptions/Source/SubtitlesAndClosedCaptions/Public/MovieSceneSubtitlesSystem.h
deleted file mode 100644
diff --git a/Engine/Plugins/Experimental/SubtitlesAndClosedCaptions/Source/SubtitlesAndClosedCaptions/Public/MovieSceneSubtitlesTrack.h b/Engine/Plugins/Experimental/SubtitlesAndClosedCaptions/Source/SubtitlesAndClosedCaptions/Public/MovieSceneSubtitlesTrack.h
deleted file mode 100644
diff --git a/Engine/Plugins/Experimental/SubtitlesAndClosedCaptions/Source/SubtitlesAndClosedCaptions/Public/SubtitleDataComponent.h b/Engine/Plugins/Experimental/SubtitlesAndClosedCaptions/Source/SubtitlesAndClosedCaptions/Public/SubtitleDataComponent.h
deleted file mode 100644
diff --git a/Engine/Plugins/Experimental/SubtitlesAndClosedCaptions/Source/SubtitlesAndClosedCaptions/Public/SubtitleWidget.h b/Engine/Plugins/Experimental/SubtitlesAndClosedCaptions/Source/SubtitlesAndClosedCaptions/Public/SubtitleWidget.h
deleted file mode 100644
diff --git a/Engine/Plugins/Experimental/SubtitlesAndClosedCaptions/Source/SubtitlesAndClosedCaptions/Public/SubtitlesAndClosedCaptionsModule.h b/Engine/Plugins/Experimental/SubtitlesAndClosedCaptions/Source/SubtitlesAndClosedCaptions/Public/SubtitlesAndClosedCaptionsModule.h
deleted file mode 100644
diff --git a/Engine/Plugins/Experimental/SubtitlesAndClosedCaptions/Source/SubtitlesAndClosedCaptions/Public/SubtitlesAudioSubsystem.h b/Engine/Plugins/Experimental/SubtitlesAndClosedCaptions/Source/SubtitlesAndClosedCaptions/Public/SubtitlesAudioSubsystem.h
deleted file mode 100644
diff --git a/Engine/Plugins/Experimental/SubtitlesAndClosedCaptions/Source/SubtitlesAndClosedCaptions/Public/SubtitlesBlueprintFunctionLibrary.h b/Engine/Plugins/Experimental/SubtitlesAndClosedCaptions/Source/SubtitlesAndClosedCaptions/Public/SubtitlesBlueprintFunctionLibrary.h
deleted file mode 100644
diff --git a/Engine/Plugins/Experimental/SubtitlesAndClosedCaptions/Source/SubtitlesAndClosedCaptions/Public/SubtitlesSettings.h b/Engine/Plugins/Experimental/SubtitlesAndClosedCaptions/Source/SubtitlesAndClosedCaptions/Public/SubtitlesSettings.h
deleted file mode 100644
diff --git a/Engine/Plugins/Experimental/SubtitlesAndClosedCaptions/Source/SubtitlesAndClosedCaptions/Public/SubtitlesSubsystem.h b/Engine/Plugins/Experimental/SubtitlesAndClosedCaptions/Source/SubtitlesAndClosedCaptions/Public/SubtitlesSubsystem.h
deleted file mode 100644
diff --git a/Engine/Plugins/Experimental/SubtitlesAndClosedCaptions/Source/SubtitlesAndClosedCaptions/SubtitlesAndClosedCaptions.Build.cs b/Engine/Plugins/Experimental/SubtitlesAndClosedCaptions/Source/SubtitlesAndClosedCaptions/SubtitlesAndClosedCaptions.Build.cs
deleted file mode 100644
diff --git a/Engine/Plugins/Experimental/SubtitlesAndClosedCaptions/Source/SubtitlesAndClosedCaptionsEditor/Private/SubtitleAssetActionUtility.cpp b/Engine/Plugins/Experimental/SubtitlesAndClosedCaptions/Source/SubtitlesAndClosedCaptionsEditor/Private/SubtitleAssetActionUtility.cpp
deleted file mode 100644
diff --git a/Engine/Plugins/Experimental/SubtitlesAndClosedCaptions/Source/SubtitlesAndClosedCaptionsEditor/Private/SubtitleAssetDefinitions.h b/Engine/Plugins/Experimental/SubtitlesAndClosedCaptions/Source/SubtitlesAndClosedCaptionsEditor/Private/SubtitleAssetDefinitions.h
deleted file mode 100644
diff --git a/Engine/Plugins/Experimental/SubtitlesAndClosedCaptions/Source/SubtitlesAndClosedCaptionsEditor/Private/SubtitleFactory.cpp b/Engine/Plugins/Experimental/SubtitlesAndClosedCaptions/Source/SubtitlesAndClosedCaptionsEditor/Private/SubtitleFactory.cpp
deleted file mode 100644
diff --git a/Engine/Plugins/Experimental/SubtitlesAndClosedCaptions/Source/SubtitlesAndClosedCaptionsEditor/Private/SubtitleFactory.h b/Engine/Plugins/Experimental/SubtitlesAndClosedCaptions/Source/SubtitlesAndClosedCaptionsEditor/Private/SubtitleFactory.h
deleted file mode 100644
diff --git a/Engine/Plugins/Experimental/SubtitlesAndClosedCaptions/Source/SubtitlesAndClosedCaptionsEditor/Private/SubtitleMigrationUtility.cpp b/Engine/Plugins/Experimental/SubtitlesAndClosedCaptions/Source/SubtitlesAndClosedCaptionsEditor/Private/SubtitleMigrationUtility.cpp
deleted file mode 100644
diff --git a/Engine/Plugins/Experimental/SubtitlesAndClosedCaptions/Source/SubtitlesAndClosedCaptionsEditor/Private/SubtitleSequencerSection.cpp b/Engine/Plugins/Experimental/SubtitlesAndClosedCaptions/Source/SubtitlesAndClosedCaptionsEditor/Private/SubtitleSequencerSection.cpp
deleted file mode 100644
diff --git a/Engine/Plugins/Experimental/SubtitlesAndClosedCaptions/Source/SubtitlesAndClosedCaptionsEditor/Private/SubtitleSequencerSection.h b/Engine/Plugins/Experimental/SubtitlesAndClosedCaptions/Source/SubtitlesAndClosedCaptionsEditor/Private/SubtitleSequencerSection.h
deleted file mode 100644
diff --git a/Engine/Plugins/Experimental/SubtitlesAndClosedCaptions/Source/SubtitlesAndClosedCaptionsEditor/Private/SubtitlesAndClosedCaptionsEditorModule.cpp b/Engine/Plugins/Experimental/SubtitlesAndClosedCaptions/Source/SubtitlesAndClosedCaptionsEditor/Private/SubtitlesAndClosedCaptionsEditorModule.cpp
deleted file mode 100644
diff --git a/Engine/Plugins/Experimental/SubtitlesAndClosedCaptions/Source/SubtitlesAndClosedCaptionsEditor/Private/SubtitlesTrackEditor.cpp b/Engine/Plugins/Experimental/SubtitlesAndClosedCaptions/Source/SubtitlesAndClosedCaptionsEditor/Private/SubtitlesTrackEditor.cpp
deleted file mode 100644
diff --git a/Engine/Plugins/Experimental/SubtitlesAndClosedCaptions/Source/SubtitlesAndClosedCaptionsEditor/Private/SubtitlesTrackEditor.h b/Engine/Plugins/Experimental/SubtitlesAndClosedCaptions/Source/SubtitlesAndClosedCaptionsEditor/Private/SubtitlesTrackEditor.h
deleted file mode 100644
diff --git a/Engine/Plugins/Experimental/SubtitlesAndClosedCaptions/Source/SubtitlesAndClosedCaptionsEditor/Public/SubtitleAssetActionUtility.h b/Engine/Plugins/Experimental/SubtitlesAndClosedCaptions/Source/SubtitlesAndClosedCaptionsEditor/Public/SubtitleAssetActionUtility.h
deleted file mode 100644
diff --git a/Engine/Plugins/Experimental/SubtitlesAndClosedCaptions/Source/SubtitlesAndClosedCaptionsEditor/Public/SubtitleMigrationUtility.h b/Engine/Plugins/Experimental/SubtitlesAndClosedCaptions/Source/SubtitlesAndClosedCaptionsEditor/Public/SubtitleMigrationUtility.h
deleted file mode 100644
diff --git a/Engine/Plugins/Experimental/SubtitlesAndClosedCaptions/Source/SubtitlesAndClosedCaptionsEditor/Public/SubtitlesAndClosedCaptionsEditorModule.h b/Engine/Plugins/Experimental/SubtitlesAndClosedCaptions/Source/SubtitlesAndClosedCaptionsEditor/Public/SubtitlesAndClosedCaptionsEditorModule.h
deleted file mode 100644
diff --git a/Engine/Plugins/Experimental/SubtitlesAndClosedCaptions/Source/SubtitlesAndClosedCaptionsEditor/SubtitlesAndClosedCaptionsEditor.Build.cs b/Engine/Plugins/Experimental/SubtitlesAndClosedCaptions/Source/SubtitlesAndClosedCaptionsEditor/SubtitlesAndClosedCaptionsEditor.Build.cs
deleted file mode 100644
diff --git a/Engine/Plugins/Experimental/SubtitlesAndClosedCaptions/SubtitlesAndClosedCaptions.uplugin b/Engine/Plugins/Experimental/SubtitlesAndClosedCaptions/SubtitlesAndClosedCaptions.uplugin
deleted file mode 100644
diff --git a/Engine/Plugins/SubtitlesAndClosedCaptions/Source/SubtitlesAndClosedCaptions/Private/MovieSceneSubtitleComponentTypes.cpp b/Engine/Plugins/SubtitlesAndClosedCaptions/Source/SubtitlesAndClosedCaptions/Private/MovieSceneSubtitleComponentTypes.cpp
new file mode 100644
--- /dev/null
+++ b/Engine/Plugins/SubtitlesAndClosedCaptions/Source/SubtitlesAndClosedCaptions/Private/MovieSceneSubtitleComponentTypes.cpp
@@ -0,0 +1,36 @@
+// Copyright Epic Games, Inc. All Rights Reserved.
+
+#include "MovieSceneSubtitleComponentTypes.h"
+
+#include "EntitySystem/MovieSceneComponentRegistry.h"
+#include "EntitySystem/MovieSceneEntityFactoryTemplates.h"
+#include "EntitySystem/MovieSceneEntitySystemLinker.h"
+
+bool FMovieSceneSubtitleComponentTypes::GMovieSceneSubtitlesComponentTypesDestroyed = false;
+TUniquePtr<FMovieSceneSubtitleComponentTypes> FMovieSceneSubtitleComponentTypes::GMovieSceneSubtitlesComponentTypes;
+
+
+FMovieSceneSubtitleComponentTypes* FMovieSceneSubtitleComponentTypes::Get()
+{
+	if (!GMovieSceneSubtitlesComponentTypes.IsValid())
+	{
+		check(!GMovieSceneSubtitlesComponentTypesDestroyed);
+		GMovieSceneSubtitlesComponentTypes.Reset(new FMovieSceneSubtitleComponentTypes);
+	}
+	return GMovieSceneSubtitlesComponentTypes.Get();
+}
+
+void FMovieSceneSubtitleComponentTypes::Destroy()
+{
+	GMovieSceneSubtitlesComponentTypes.Reset();
+	GMovieSceneSubtitlesComponentTypesDestroyed = true;
+}
+
+FMovieSceneSubtitleComponentTypes::FMovieSceneSubtitleComponentTypes()
+{
+	using namespace UE::MovieScene;
+	FComponentRegistry* ComponentRegistry = UMovieSceneEntitySystemLinker::GetComponents();
+	check(ComponentRegistry);
+	
+	ComponentRegistry->NewComponentType(&SubtitleData, TEXT("SubtitleData Component"));
+}
diff --git a/Engine/Plugins/SubtitlesAndClosedCaptions/Source/SubtitlesAndClosedCaptions/Private/MovieSceneSubtitleSection.cpp b/Engine/Plugins/SubtitlesAndClosedCaptions/Source/SubtitlesAndClosedCaptions/Private/MovieSceneSubtitleSection.cpp
new file mode 100644
--- /dev/null
+++ b/Engine/Plugins/SubtitlesAndClosedCaptions/Source/SubtitlesAndClosedCaptions/Private/MovieSceneSubtitleSection.cpp
@@ -0,0 +1,68 @@
+// Copyright Epic Games, Inc. All Rights Reserved.
+
+#include "MovieSceneSubtitleSection.h"
+
+#include "EntitySystem/BuiltInComponentTypes.h"
+#include "EntitySystem/MovieSceneEntitySystemLinker.h"
+#include "MovieScene.h"
+#include "MovieSceneFwd.h"
+#include "MovieSceneSubtitleComponentTypes.h"
+#include "SubtitleDataComponent.h"
+#include "Subtitles/SubtitlesAndClosedCaptionsDelegates.h"
+#include "Subtitles/SubtitlesAndClosedCaptionsTypes.h"
+
+#include UE_INLINE_GENERATED_CPP_BY_NAME(MovieSceneSubtitleSection)
+
+using namespace UE::MovieScene;
+
+UMovieSceneSubtitleSection::UMovieSceneSubtitleSection(const FObjectInitializer& ObjectInitializer)
+	: Super(ObjectInitializer)
+{
+}
+
+TOptional<TRange<FFrameNumber>> UMovieSceneSubtitleSection::GetAutoSizeRange() const
+{
+	if (!Subtitle)
+	{
+		return TRange<FFrameNumber>();
+	}
+
+	// Find the maximum duration of all the subtitles in the UserData, including their delayed start offsets.
+	const FFrameRate FrameRate = GetTypedOuter<UMovieScene>()->GetTickResolution();
+	const FFrameTime DurationToUse = Subtitle->GetMaximumDuration() * FrameRate;
+	return TRange<FFrameNumber>(GetInclusiveStartFrame(), GetInclusiveStartFrame() + DurationToUse.FrameNumber);
+}
+
+bool UMovieSceneSubtitleSection::PopulateEvaluationFieldImpl(const TRange<FFrameNumber>& EffectiveRange, const FMovieSceneEvaluationFieldEntityMetaData& InMetaData, FMovieSceneEntityComponentFieldBuilder* OutFieldBuilder)
+{
+	// check preconditions
+	check(OutFieldBuilder);
+
+	OutFieldBuilder->AddPersistentEntity(EffectiveRange, this);
+	return true;
+}
+
+void UMovieSceneSubtitleSection::ImportEntityImpl(UMovieSceneEntitySystemLinker* EntityLinker, const FEntityImportParams& Params, FImportedEntity* OutImportedEntity)
+{
+	// check preconditions
+	check(OutImportedEntity);
+
+	if (!Subtitle)
+	{
+		return;
+	}
+
+	const FBuiltInComponentTypes* BuiltInComponents = FBuiltInComponentTypes::Get();
+	check(BuiltInComponents);
+
+	FMovieSceneSubtitleComponentTypes* SubtitleComponents = FMovieSceneSubtitleComponentTypes::Get();
+	check(SubtitleComponents);
+
+	const FGuid ObjectBindingID = Params.GetObjectBindingID();
+
+	OutImportedEntity->AddBuilder(
+		FEntityBuilder()
+		.AddConditional(BuiltInComponents->GenericObjectBinding, ObjectBindingID, ObjectBindingID.IsValid())
+		.AddTagConditional(BuiltInComponents->Tags.Root, !ObjectBindingID.IsValid())
+		.Add(SubtitleComponents->SubtitleData, FSubtitleDataComponent{ this, EMovieScenePlayerStatus::Type::Stopped }));
+}
diff --git a/Engine/Plugins/SubtitlesAndClosedCaptions/Source/SubtitlesAndClosedCaptions/Private/MovieSceneSubtitlesSystem.cpp b/Engine/Plugins/SubtitlesAndClosedCaptions/Source/SubtitlesAndClosedCaptions/Private/MovieSceneSubtitlesSystem.cpp
new file mode 100644
--- /dev/null
+++ b/Engine/Plugins/SubtitlesAndClosedCaptions/Source/SubtitlesAndClosedCaptions/Private/MovieSceneSubtitlesSystem.cpp
@@ -0,0 +1,70 @@
+// Copyright Epic Games, Inc. All Rights Reserved.
+
+#include "MovieSceneSubtitlesSystem.h"
+
+#include "EntitySystem/BuiltInComponentTypes.h"
+
+#include "MovieSceneSubtitleComponentTypes.h"
+#include "MovieSceneSubtitleSection.h"
+
+#include "SubtitlesAndClosedCaptionsModule.h"
+
+#include UE_INLINE_GENERATED_CPP_BY_NAME(MovieSceneSubtitlesSystem)
+
+using namespace UE::MovieScene;
+
+
+UMovieSceneSubtitlesSystem::UMovieSceneSubtitlesSystem(const FObjectInitializer& ObjInit)
+	: UMovieSceneEntitySystem(ObjInit)
+{
+	FMovieSceneSubtitleComponentTypes* SubtitlesComponents = FMovieSceneSubtitleComponentTypes::Get();
+	check(SubtitlesComponents);
+
+	RelevantComponent = SubtitlesComponents->SubtitleData;
+
+	Phase = ESystemPhase::Scheduling;
+}
+
+UMovieSceneSubtitlesSystem::~UMovieSceneSubtitlesSystem()
+{
+}
+
+void UMovieSceneSubtitlesSystem::OnSchedulePersistentTasks(UE::MovieScene::IEntitySystemScheduler* TaskScheduler)
+{
+	const FBuiltInComponentTypes* BuiltInComponents = FBuiltInComponentTypes::Get();
+	check(BuiltInComponents);
+
+	const FMovieSceneSubtitleComponentTypes* SubtitleComponents = FMovieSceneSubtitleComponentTypes::Get();
+	check(SubtitleComponents);
+
+	check(Linker);
+	const FInstanceRegistry* InstanceRegistry = Linker->GetInstanceRegistry();
+	check(InstanceRegistry);
+
+	FTaskID EvaluateTask = FEntityTaskBuilder()
+		.Read(BuiltInComponents->InstanceHandle)
+		.Write(SubtitleComponents->SubtitleData)
+		.SetDesiredThread(ENamedThreads::GameThread)
+		.Schedule_PerAllocation<FEvaluateSubtitles>(&Linker->EntityManager, TaskScheduler, *InstanceRegistry);
+}
+
+void UMovieSceneSubtitlesSystem::OnRun(FSystemTaskPrerequisites& InPrerequisites, FSystemSubsequentTasks& Subsequents)
+{
+	const FBuiltInComponentTypes* BuiltInComponents = FBuiltInComponentTypes::Get();
+	check(BuiltInComponents);
+
+	const FMovieSceneSubtitleComponentTypes* SubtitleComponents = FMovieSceneSubtitleComponentTypes::Get();
+	check(SubtitleComponents);
+
+	check(Linker);
+	const FInstanceRegistry* InstanceRegistry = Linker->GetInstanceRegistry();
+	check(InstanceRegistry);
+
+	FEntityTaskBuilder()
+		.Read(BuiltInComponents->InstanceHandle)
+		.Write(SubtitleComponents->SubtitleData)
+		.SetDesiredThread(ENamedThreads::GameThread)
+		.template Dispatch_PerAllocation<FEvaluateSubtitles>(&Linker->EntityManager, InPrerequisites, &Subsequents, *InstanceRegistry);
+
+
+}
diff --git a/Engine/Plugins/SubtitlesAndClosedCaptions/Source/SubtitlesAndClosedCaptions/Private/MovieSceneSubtitlesTrack.cpp b/Engine/Plugins/SubtitlesAndClosedCaptions/Source/SubtitlesAndClosedCaptions/Private/MovieSceneSubtitlesTrack.cpp
new file mode 100644
--- /dev/null
+++ b/Engine/Plugins/SubtitlesAndClosedCaptions/Source/SubtitlesAndClosedCaptions/Private/MovieSceneSubtitlesTrack.cpp
@@ -0,0 +1,104 @@
+// Copyright Epic Games, Inc. All Rights Reserved.
+
+#include "MovieSceneSubtitlesTrack.h"
+
+#include "MovieScene.h"
+#include "MovieSceneCommonHelpers.h"
+#include "MovieSceneSubtitleSection.h"
+#include "Subtitles/SubtitlesAndClosedCaptionsDelegates.h"
+#include "Subtitles/SubtitlesAndClosedCaptionsTypes.h"
+
+#include UE_INLINE_GENERATED_CPP_BY_NAME(MovieSceneSubtitlesTrack)
+
+#define LOCTEXT_NAMESPACE "MovieSceneSubtitlesTrack"
+
+constexpr float DefaultSubtitleSectionDuration = 3.f;
+
+UMovieSceneSubtitlesTrack::UMovieSceneSubtitlesTrack(const FObjectInitializer& ObjectInitializer)
+	: Super(ObjectInitializer)
+{
+}
+
+UMovieSceneSection* UMovieSceneSubtitlesTrack::AddNewSubtitle(const USubtitleAssetUserData& Subtitle, FFrameNumber InStartTime)
+{
+	return AddNewSubtitleOnRow(Subtitle, InStartTime, INDEX_NONE);
+}
+
+UMovieSceneSection* UMovieSceneSubtitlesTrack::AddNewSubtitleOnRow(const USubtitleAssetUserData& Subtitle, FFrameNumber InStartTime, int32 InRowIndex)
+{
+	UMovieSceneSubtitleSection* NewSection = Cast<UMovieSceneSubtitleSection>(CreateNewSection());
+
+	const FFrameRate FrameRate = GetTypedOuter<UMovieScene>()->GetTickResolution();
+	const float SubtitleDuration = Subtitle.GetMaximumDuration();
+
+	// Find the maximum duration of all the subtitles in the UserData, including their delayed start offsets.
+	const bool UseDefaultDuration = Subtitle.Subtitles.IsEmpty();
+	// Use the default duration if the subtitles array is empty.
+	// If it's not empty, the subtitles themselves have a minimum duration and therefore have a non-zero length.
+	const FFrameTime DurationToUse = UseDefaultDuration ? (DefaultSubtitleSectionDuration * FrameRate) : (SubtitleDuration * FrameRate);
+
+	NewSection->InitialPlacementOnRow(SubtitleSections, InStartTime, DurationToUse.FrameNumber.Value, InRowIndex);
+	NewSection->SetSubtitle(Subtitle);
+
+	// If the duration was 0-length, turn orange to make it clear that the duration was defaulted.
+	if (UseDefaultDuration)
+	{
+		NewSection->SetColorTint(FColor::Orange);
+	}
+
+	SubtitleSections.Add(NewSection);
+
+	return NewSection;
+}
+
+bool UMovieSceneSubtitlesTrack::SupportsType(TSubclassOf<UMovieSceneSection> SectionClass) const
+{
+	return SectionClass == UMovieSceneSubtitleSection::StaticClass();
+}
+
+void UMovieSceneSubtitlesTrack::RemoveAllAnimationData()
+{
+	return SubtitleSections.Empty();
+}
+
+bool UMovieSceneSubtitlesTrack::HasSection(const UMovieSceneSection& Section) const
+{
+	return SubtitleSections.Contains(&Section);
+}
+
+void UMovieSceneSubtitlesTrack::AddSection(UMovieSceneSection& Section)
+{
+	SubtitleSections.Add(&Section);
+}
+
+void UMovieSceneSubtitlesTrack::RemoveSection(UMovieSceneSection& Section)
+{
+	SubtitleSections.Remove(&Section);
+}
+
+void UMovieSceneSubtitlesTrack::RemoveSectionAt(int32 SectionIndex)
+{
+	SubtitleSections.RemoveAt(SectionIndex);
+}
+
+bool UMovieSceneSubtitlesTrack::IsEmpty() const
+{
+	return SubtitleSections.IsEmpty();
+}
+
+const TArray<UMovieSceneSection*>& UMovieSceneSubtitlesTrack::GetAllSections() const
+{
+	return SubtitleSections;
+}
+
+bool UMovieSceneSubtitlesTrack::SupportsMultipleRows() const
+{
+	return true;
+}
+
+UMovieSceneSection* UMovieSceneSubtitlesTrack::CreateNewSection()
+{
+	return NewObject<UMovieSceneSubtitleSection>(this, NAME_None, RF_Transactional);
+}
+
+#undef LOCTEXT_NAMESPACE
diff --git a/Engine/Plugins/SubtitlesAndClosedCaptions/Source/SubtitlesAndClosedCaptions/Private/SubtitleWidget.cpp b/Engine/Plugins/SubtitlesAndClosedCaptions/Source/SubtitlesAndClosedCaptions/Private/SubtitleWidget.cpp
new file mode 100644
--- /dev/null
+++ b/Engine/Plugins/SubtitlesAndClosedCaptions/Source/SubtitlesAndClosedCaptions/Private/SubtitleWidget.cpp
@@ -0,0 +1,161 @@
+// Copyright Epic Games, Inc. All Rights Reserved.
+
+#include "SubtitleWidget.h"
+#include "Components/Border.h"
+#include "Components/TextBlock.h"
+#include "Logging/StructuredLog.h"
+#include "SubtitlesAndClosedCaptionsModule.h"
+
+#include UE_INLINE_GENERATED_CPP_BY_NAME(SubtitleWidget)
+
+void USubtitleWidget::NativeConstruct()
+{
+	Super::NativeConstruct();
+
+	StopDisplayingAllSubtitles();
+}
+
+void USubtitleWidget::StartDisplayingSubtitle(const FSubtitleAssetData& Subtitle)
+{
+	// Pick which UextBlock is relevant by category....
+	UTextBlock* TextToModify = nullptr;
+	UBorder* BorderToModify = nullptr;
+
+	const ESubtitleType SubtitleType = Subtitle.SubtitleType;
+
+	switch (SubtitleType)
+	{
+	default:
+	{
+		UE_LOGFMT(LogSubtitlesAndClosedCaptions, Warning, "An unrecognized subtitle type was queued for display. Using the standard Subtitle text block as a fallback.");
+		// Fallthrough, no break.
+	}
+
+	case ESubtitleType::Subtitle:
+	{
+		TextToModify = DialogSubtitleBlock;
+		BorderToModify = DialogBorder;
+		break;
+	}
+
+	case ESubtitleType::ClosedCaption:
+	{
+		TextToModify = CaptionSubtitleBlock;
+		BorderToModify = CaptionBorder;
+		break;
+	}
+
+	case ESubtitleType::AudioDescription:
+	{
+		TextToModify = DescriptionSubtitleBlock;
+		BorderToModify = DescriptionBorder;
+		break;
+	}
+
+	}
+
+	// ...Then modify and display it
+	if (IsValid(TextToModify))
+	{
+		TextToModify->SetText(Subtitle.Text);
+		TextToModify->SetVisibility(ESlateVisibility::HitTestInvisible);
+
+		// Only hide/display border if the corresponding text is also valid.
+		if (IsValid(BorderToModify))
+		{
+			BorderToModify->SetVisibility(ESlateVisibility::HitTestInvisible);
+		}
+
+		// Track the currently-displayed subtitle for comparison when removing later
+		CurrentlyDisplayedSubtitles[SubtitleType] = Subtitle;
+	}
+}
+
+void USubtitleWidget::StopDisplayingSubtitle(const FSubtitleAssetData& Subtitle)
+{
+	UTextBlock* TextToModify = nullptr;
+	UBorder* BorderToModify = nullptr;
+
+	// Confirm that the removed subtitle is actually being displayed before trying to stop it.
+	const ESubtitleType SubtitleType = Subtitle.SubtitleType;
+	if (CurrentlyDisplayedSubtitles[SubtitleType].IsSet() && CurrentlyDisplayedSubtitles[SubtitleType] == Subtitle)
+	{
+		switch (SubtitleType)
+		{
+		default: // Fallthrough, no break (see ::StartDisplayingSubtitle)
+		case ESubtitleType::Subtitle:
+		{
+			TextToModify = DialogSubtitleBlock;
+			BorderToModify = DialogBorder;
+			break;
+		}
+
+		case ESubtitleType::ClosedCaption:
+		{
+			TextToModify = CaptionSubtitleBlock;
+			BorderToModify = CaptionBorder;
+			break;
+		}
+
+		case ESubtitleType::AudioDescription:
+		{
+			TextToModify = DescriptionSubtitleBlock;
+			BorderToModify = DescriptionBorder;
+			break;
+		}
+
+		}
+
+		if (IsValid(TextToModify))
+		{
+			TextToModify->SetVisibility(ESlateVisibility::Collapsed);
+
+			if (IsValid(BorderToModify))
+			{
+				BorderToModify->SetVisibility(ESlateVisibility::Collapsed);
+			}
+		}
+	}
+}
+
+void USubtitleWidget::StopDisplayingAllSubtitles()
+{
+	// Ensure the text blocks don't start visible.
+	if (IsValid(DialogSubtitleBlock))
+	{
+		DialogSubtitleBlock->SetVisibility(ESlateVisibility::Collapsed);
+	}
+
+	if (IsValid(CaptionSubtitleBlock))
+	{
+		CaptionSubtitleBlock->SetVisibility(ESlateVisibility::Collapsed);
+	}
+
+	if (IsValid(DescriptionSubtitleBlock))
+	{
+		DescriptionSubtitleBlock->SetVisibility(ESlateVisibility::Collapsed);
+	}
+
+	// Borders should also not be visible on startup.
+	if (IsValid(DialogBorder))
+	{
+		DialogBorder->SetVisibility(ESlateVisibility::Collapsed);
+	}
+	if (IsValid(CaptionBorder))
+	{
+		CaptionBorder->SetVisibility(ESlateVisibility::Collapsed);
+	}
+	if (IsValid(DescriptionBorder))
+	{
+		DescriptionBorder->SetVisibility(ESlateVisibility::Collapsed);
+	}
+}
+
+USubtitleWidget::USubtitleWidget(const FObjectInitializer& ObjectInitializer)
+	: UUserWidget(ObjectInitializer)
+{
+
+	CurrentlyDisplayedSubtitles.Add({ ESubtitleType::Subtitle,			TOptional<FSubtitleAssetData>() });
+	CurrentlyDisplayedSubtitles.Add({ ESubtitleType::AudioDescription,	TOptional<FSubtitleAssetData>() });
+	CurrentlyDisplayedSubtitles.Add({ ESubtitleType::ClosedCaption,		TOptional<FSubtitleAssetData>() });
+}
diff --git a/Engine/Plugins/SubtitlesAndClosedCaptions/Source/SubtitlesAndClosedCaptions/Private/SubtitlesAndClosedCaptionsModule.cpp b/Engine/Plugins/SubtitlesAndClosedCaptions/Source/SubtitlesAndClosedCaptions/Private/SubtitlesAndClosedCaptionsModule.cpp
new file mode 100644
--- /dev/null
+++ b/Engine/Plugins/SubtitlesAndClosedCaptions/Source/SubtitlesAndClosedCaptions/Private/SubtitlesAndClosedCaptionsModule.cpp
@@ -0,0 +1,26 @@
+// Copyright Epic Games, Inc. All Rights Reserved.
+
+#include "SubtitlesAndClosedCaptionsModule.h"
+
+#include "Features/IModularFeatures.h"
+#include "Modules/ModuleManager.h"
+#include "Subtitles/SubtitlesAndClosedCaptionsDelegates.h"
+
+DEFINE_LOG_CATEGORY(LogSubtitlesAndClosedCaptions);
+
+void FSubtitlesAndClosedCaptionsModule::StartupModule()
+{
+	IModularFeatures::Get().RegisterModularFeature(ISubtitlesAndClosedCaptionsModule::GetModularFeatureName(), this);
+}
+
+void FSubtitlesAndClosedCaptionsModule::ShutdownModule()
+{
+	IModularFeatures::Get().UnregisterModularFeature(ISubtitlesAndClosedCaptionsModule::GetModularFeatureName(), this);
+}
+
+TSubclassOf<UAssetUserData> FSubtitlesAndClosedCaptionsModule::GetAssetUserDataClass() const
+{
+	return USubtitleAssetUserData::StaticClass();
+}
+
+IMPLEMENT_MODULE(FSubtitlesAndClosedCaptionsModule, SubtitlesAndClosedCaptions);
diff --git a/Engine/Plugins/SubtitlesAndClosedCaptions/Source/SubtitlesAndClosedCaptions/Private/SubtitlesAudioSubsystem.cpp b/Engine/Plugins/SubtitlesAndClosedCaptions/Source/SubtitlesAndClosedCaptions/Private/SubtitlesAudioSubsystem.cpp
new file mode 100644
--- /dev/null
+++ b/Engine/Plugins/SubtitlesAndClosedCaptions/Source/SubtitlesAndClosedCaptions/Private/SubtitlesAudioSubsystem.cpp
@@ -0,0 +1,188 @@
+// Copyright Epic Games, Inc. All Rights Reserved.
+
+#include "SubtitlesAudioSubsystem.h"
+
+#include "ActiveSound.h"
+#include "Sound/DialogueSoundWaveProxy.h"
+#include "Sound/SoundCue.h"
+#include "Sound/SoundNodeDialoguePlayer.h"
+#include "Subtitles/SubtitlesAndClosedCaptionsTypes.h"
+#include "Subtitles/SubtitlesAndClosedCaptionsDelegates.h"
+
+#include UE_INLINE_GENERATED_CPP_BY_NAME(SubtitlesAudioSubsystem)
+
+void USubtitlesAudioSubsystem::NotifyActiveSoundCreated(FActiveSound& ActiveSound)
+{
+	USoundBase* Sound = ActiveSound.GetSound();
+
+	if (IsValid(Sound))
+	{
+		if (Sound->HasAssetUserDataOfClass(USubtitleAssetUserData::StaticClass()))
+		{
+			const USubtitleAssetUserData* SubtitleAsset = CastChecked<USubtitleAssetUserData>(Sound->GetAssetUserDataOfClass(USubtitleAssetUserData::StaticClass()));
+			for (const FSubtitleAssetData& Subtitle : SubtitleAsset->Subtitles)
+			{
+				const bool bUseDurationProperty = (Subtitle.SubtitleDurationType == ESubtitleDurationType::UseDurationProperty);
+				FSubtitlesAndClosedCaptionsDelegates::QueueSubtitle.ExecuteIfBound(FQueueSubtitleParameters{ Subtitle }, bUseDurationProperty ? ESubtitleTiming::InternallyTimed : ESubtitleTiming::ExternallyTimed);
+			}
+		}
+
+		// Also check for DialogueWaves with localized text
+		// DialogueWaves, even inside SoundCues, should only play subtitles this way if the migration cvar was set.
+		// Since this callback is run when the sound is created in the sound engine, changing the cvar doesn't effect currently-playing sounds.
+		static const IConsoleVariable* CVarUseNewSubtitles = IConsoleManager::Get().FindConsoleVariable(TEXT("au.UseNewSubtitles"));
+
+		if (CVarUseNewSubtitles && CVarUseNewSubtitles->GetBool() && ActiveSound.bHandleSubtitles)
+		{
+			FSubtitleAssetData SubtitleCue;
+
+			const UDialogueSoundWaveProxy* DialogueWave = Cast<UDialogueSoundWaveProxy>(Sound);
+			if (IsValid(DialogueWave))
+			{
+				QueueSubtitleFromDialogueWaveProxy(DialogueWave, ActiveSound);
+			}
+			else
+			{
+				// Check inside Sound Cues for DialogueWaves as well.
+				const USoundCue* SoundCue = Cast<USoundCue>(Sound);
+
+				// SoundCue's FirstNode is used to grab all the other nodes, so it needs to be valid also.
+				if (IsValid(SoundCue) && IsValid(SoundCue->FirstNode))
+				{
+					const UDialogueSoundWaveProxy* Proxy = CheckSoundCueNodeForDialogueWaves(SoundCue->FirstNode, ActiveSound);
+
+					if (IsValid(Proxy))
+					{
+						QueueSubtitleFromDialogueWaveProxy(Proxy, ActiveSound);
+					}
+					else
+					{
+						// If the parent node is not a DialogueWavePlayer, check for nested nodes that may also have DialogueWaves
+						TArray<USoundNode*> ChildNodes;
+						SoundCue->FirstNode->GetAllNodes(ChildNodes);
+						for (const USoundNode* Node : ChildNodes)
+						{
+							const UDialogueSoundWaveProxy* ChildProxy = CheckSoundCueNodeForDialogueWaves(Node, ActiveSound);
+							if (IsValid(ChildProxy))
+							{
+								QueueSubtitleFromDialogueWaveProxy(ChildProxy, ActiveSound);
+							}
+						}
+					}
+				}
+			}
+		}
+	}
+}
+
+void USubtitlesAudioSubsystem::NotifyActiveSoundDeleting(const FActiveSound& ActiveSound)
+{
+	USoundBase* Sound = ActiveSound.GetSound();
+
+	if (IsValid(Sound))
+	{
+		if (Sound->HasAssetUserDataOfClass(USubtitleAssetUserData::StaticClass()))
+		{
+			const USubtitleAssetUserData* SubtitleAsset = CastChecked<USubtitleAssetUserData>(Sound->GetAssetUserDataOfClass(USubtitleAssetUserData::StaticClass()));
+			for (const FSubtitleAssetData& Subtitle : SubtitleAsset->Subtitles)
+			{
+				const bool bUseDurationProperty = (Subtitle.SubtitleDurationType == ESubtitleDurationType::UseDurationProperty);
+				// If not using sound duration, subtitles continue to play instead of stopping:
+				// Their duration is completely decoupled from the sound's.
+
+				// ...but if Using the sound's duration: Stop All subtitles.
+				if (!bUseDurationProperty)
+				{
+					FSubtitlesAndClosedCaptionsDelegates::StopSubtitle.ExecuteIfBound(Subtitle);
+				}
+			}
+		}
+
+		// Also check for DialogueWaves with localized text
+		static const IConsoleVariable* CVarUseNewSubtitles = IConsoleManager::Get().FindConsoleVariable(TEXT("au.UseNewSubtitles"));
+
+		if (CVarUseNewSubtitles && CVarUseNewSubtitles->GetBool() && ActiveSound.bHandleSubtitles)
+		{
+			FSubtitleAssetData SubtitleCue;
+
+			const UDialogueSoundWaveProxy* DialogueWave = Cast<UDialogueSoundWaveProxy>(Sound);
+			if (IsValid(DialogueWave))
+			{
+				StopSubtitleFromDialogueWaveProxy(DialogueWave, ActiveSound);
+			}
+			else
+			{
+				// Check inside Sound Cues for DialogueWaves as well.
+				const USoundCue* SoundCue = Cast<USoundCue>(Sound);
+
+				// SoundCue's FirstNode is used to grab all the other nodes, so it needs to be valid also.
+				if (IsValid(SoundCue) && IsValid(SoundCue->FirstNode))
+				{
+					const UDialogueSoundWaveProxy* Proxy = CheckSoundCueNodeForDialogueWaves(SoundCue->FirstNode, ActiveSound);
+
+					if (IsValid(Proxy))
+					{
+						// A dialogueWave won't reference another dialoguewave, so no need to search recursively
+						StopSubtitleFromDialogueWaveProxy(Proxy, ActiveSound);
+					}
+					else
+					{
+						// If the parent node is not a DialogueWavePlayer, check for nested nodes that may also have DialogueWaves
+						TArray<USoundNode*> ChildNodes;
+						SoundCue->FirstNode->GetAllNodes(ChildNodes);
+						for (const USoundNode* Node : ChildNodes)
+						{
+							const UDialogueSoundWaveProxy* ChildProxy = CheckSoundCueNodeForDialogueWaves(Node, ActiveSound);
+							if (IsValid(ChildProxy))
+							{
+								StopSubtitleFromDialogueWaveProxy(ChildProxy, ActiveSound);
+							}
+						}
+					}
+				}
+			}
+		}
+	}
+}
+
+const UDialogueSoundWaveProxy* USubtitlesAudioSubsystem::CheckSoundCueNodeForDialogueWaves(const USoundNode* Node, const FActiveSound& ActiveSound)
+{
+	const USoundNodeDialoguePlayer* DialoguePlayerNode = Cast<USoundNodeDialoguePlayer>(Node);
+
+	// These pointers might be null or invalid:
+	// A SoundCue might have a Dialogue node that doesn't actually reference any assets, or the asset itself might lack any actual soundwaves to play.
+	if (IsValid(DialoguePlayerNode))
+	{
+		const UDialogueSoundWaveProxy* Proxy = DialoguePlayerNode->GetContextualDialogueSoundWaveProxy();
+		if (IsValid(Proxy))
+		{
+			return Proxy;
+		}
+	}
+
+	return nullptr;
+}
+
+void USubtitlesAudioSubsystem::QueueSubtitleFromDialogueWaveProxy(const UDialogueSoundWaveProxy* DialogueWaveProxy, const FActiveSound& ActiveSound)
+{
+	check(ActiveSound.bHandleSubtitles);	// Function should only be called if the Active Sound actually has subtitles it needs queued.
+	FSubtitleAssetData SubtitleCue;
+	SubtitleCue.Text = DialogueWaveProxy->GetSubtitleText();
+	SubtitleCue.Priority = ActiveSound.SubtitlePriority;
+	SubtitleCue.Duration = DialogueWaveProxy->GetDuration();
+	SubtitleCue.StartOffset = ActiveSound.RequestedStartTime;
+
+	FSubtitlesAndClosedCaptionsDelegates::QueueSubtitle.ExecuteIfBound(FQueueSubtitleParameters{ SubtitleCue }, ESubtitleTiming::InternallyTimed);
+}
+
+void USubtitlesAudioSubsystem::StopSubtitleFromDialogueWaveProxy(const UDialogueSoundWaveProxy* DialogueWaveProxy, const FActiveSound& ActiveSound)
+{
+	check(ActiveSound.bHandleSubtitles);	// Function should only be called if the Active Sound actually has subtitles it needs queued.
+	FSubtitleAssetData SubtitleCue;
+	SubtitleCue.Text = DialogueWaveProxy->GetSubtitleText();
+	SubtitleCue.Priority = ActiveSound.SubtitlePriority;
+	SubtitleCue.Duration = DialogueWaveProxy->GetDuration();
+	SubtitleCue.StartOffset = ActiveSound.RequestedStartTime;
+
+	FSubtitlesAndClosedCaptionsDelegates::StopSubtitle.ExecuteIfBound(SubtitleCue);
+}
diff --git a/Engine/Plugins/SubtitlesAndClosedCaptions/Source/SubtitlesAndClosedCaptions/Private/SubtitlesBlueprintFunctionLibrary.cpp b/Engine/Plugins/SubtitlesAndClosedCaptions/Source/SubtitlesAndClosedCaptions/Private/SubtitlesBlueprintFunctionLibrary.cpp
new file mode 100644
--- /dev/null
+++ b/Engine/Plugins/SubtitlesAndClosedCaptions/Source/SubtitlesAndClosedCaptions/Private/SubtitlesBlueprintFunctionLibrary.cpp
@@ -0,0 +1,105 @@
+// Copyright Epic Games, Inc. All Rights Reserved.
+
+#include "SubtitlesBlueprintFunctionLibrary.h"
+
+#include "Engine/Engine.h"
+#include "Subsystems/SubsystemBlueprintLibrary.h"
+#include "Subtitles/SubtitlesAndClosedCaptionsDelegates.h"
+#include "SubtitlesSubsystem.h"
+
+#include UE_INLINE_GENERATED_CPP_BY_NAME(SubtitlesBlueprintFunctionLibrary)
+
+void USubtitlesBlueprintFunctionLibrary::QueueSubtitlesFromAsset(TSoftObjectPtr<USubtitleAssetUserData> SubtitleAsset)
+{
+	if (SubtitleAsset.IsValid())
+	{
+		for (const FSubtitleAssetData& Subtitle : SubtitleAsset->Subtitles)
+		{
+			FQueueSubtitleParameters Params{ Subtitle };
+
+			// Force internal timing: external timing requires the caller to manually start and stop each subtitle in the asset, defeating this function's purpose.
+			// For manual timing control on a per-subtitle basis, QueueSubtitle exists.
+			FSubtitlesAndClosedCaptionsDelegates::QueueSubtitle.ExecuteIfBound(Params, ESubtitleTiming::InternallyTimed);
+		}
+	}
+}
+
+void USubtitlesBlueprintFunctionLibrary::StopSubtitlesInAsset(TSoftObjectPtr<USubtitleAssetUserData> SubtitleAsset)
+{
+	if (SubtitleAsset.IsValid())
+	{
+		for (const FSubtitleAssetData& Subtitle : SubtitleAsset->Subtitles)
+		{
+			FSubtitlesAndClosedCaptionsDelegates::StopSubtitle.ExecuteIfBound(Subtitle);
+		}
+	}
+}
+
+// Old deprecated version of QueueSingleSubtitle; use that instead.
+void USubtitlesBlueprintFunctionLibrary::QueueSubtitle(const FSubtitleAssetData& Subtitle, const ESubtitleTiming Timing)
+{
+	FQueueSubtitleParameters Params{ Subtitle };
+	FSubtitlesAndClosedCaptionsDelegates::QueueSubtitle.ExecuteIfBound(Params, Timing);
+}
+
+void USubtitlesBlueprintFunctionLibrary::QueueSingleSubtitle(
+	const FText Text, const float Duration, const float StartOffset, const float Priority, const ESubtitleType SubtitleType, const ESubtitleTiming Timing)
+{
+	FSubtitleAssetData Subtitle{
+		.Text = Text,
+		.SubtitleDurationType = ESubtitleDurationType::UseDurationProperty,	// There is no sound to synchronize to, so use the provided duration.
+		.Duration = Duration,
+		.StartOffset = StartOffset,
+		.Priority = Priority,
+		.SubtitleType = SubtitleType,
+	};
+	FQueueSubtitleParameters Params{ Subtitle };
+	FSubtitlesAndClosedCaptionsDelegates::QueueSubtitle.ExecuteIfBound(Params, Timing);
+}
+
+bool USubtitlesBlueprintFunctionLibrary::IsSubtitleActive(const FSubtitleAssetData& Subtitle)
+{
+	if (!FSubtitlesAndClosedCaptionsDelegates::IsSubtitleActive.IsBound())
+	{
+		return false;
+	}
+	return FSubtitlesAndClosedCaptionsDelegates::IsSubtitleActive.Execute(Subtitle);
+}
+
+void USubtitlesBlueprintFunctionLibrary::StopSubtitle(const FSubtitleAssetData& Subtitle)
+{
+	FSubtitlesAndClosedCaptionsDelegates::StopSubtitle.ExecuteIfBound(Subtitle);
+}
+
+void USubtitlesBlueprintFunctionLibrary::StopAllSubtitles()
+{
+	FSubtitlesAndClosedCaptionsDelegates::StopAllSubtitles.ExecuteIfBound();
+}
+
+void USubtitlesBlueprintFunctionLibrary::ReplaceSubtitleWidget(const TSoftClassPtr<USubtitleWidget>& NewWidgetAsset)
+{
+	// ReplaceSubtitleWidget doesn't have a corresponding delegate call (no need to expose it to the engine), so we need to get the subsystem manually as a result.
+	if (GEngine == nullptr)
+	{
+		return;
+	}
+
+	// Get the subtitles subsystem from the visible world.
+	if (const UGameViewportClient* Viewport = GEngine->GameViewport)
+	{
+		const FWorldContext& WorldContext = GEngine->GetWorldContextFromGameViewportChecked(Viewport);
+		UWorld* World = WorldContext.World();
+
+		if (IsValid(World))
+		{
+			// ReplaceWidget is non-const, so no const pointers here.
+			UWorldSubsystem* WorldSubsystem = USubsystemBlueprintLibrary::GetWorldSubsystem(World, USubtitlesSubsystem::StaticClass());
+			USubtitlesSubsystem* SubtitlesSubsystem = Cast<USubtitlesSubsystem>(WorldSubsystem);
+
+			if (IsValid(SubtitlesSubsystem))
+			{
+				SubtitlesSubsystem->ReplaceWidget(NewWidgetAsset);
+			}
+		}
+	}
+}
diff --git a/Engine/Plugins/SubtitlesAndClosedCaptions/Source/SubtitlesAndClosedCaptions/Private/SubtitlesSettings.cpp b/Engine/Plugins/SubtitlesAndClosedCaptions/Source/SubtitlesAndClosedCaptions/Private/SubtitlesSettings.cpp
new file mode 100644
--- /dev/null
+++ b/Engine/Plugins/SubtitlesAndClosedCaptions/Source/SubtitlesAndClosedCaptions/Private/SubtitlesSettings.cpp
@@ -0,0 +1,14 @@
+// Copyright Epic Games, Inc. All Rights Reserved.
+#include "SubtitlesSettings.h"
+
+#include UE_INLINE_GENERATED_CPP_BY_NAME(SubtitlesSettings)
+
+#define LOCTEXT_NAMESPACE "SubtitlesSettings"
+
+USubtitlesSettings::USubtitlesSettings()
+{
+	WidgetToUseAsSubclass = SubtitleWidgetToUse.Get();
+	DefaultWidgetToUseAsSubclass = SubtitleWidgetToUseDefault.Get();
+}
+
+#undef LOCTEXT_NAMESPACE
diff --git a/Engine/Plugins/SubtitlesAndClosedCaptions/Source/SubtitlesAndClosedCaptions/Private/SubtitlesSubsystem.cpp b/Engine/Plugins/SubtitlesAndClosedCaptions/Source/SubtitlesAndClosedCaptions/Private/SubtitlesSubsystem.cpp
new file mode 100644
--- /dev/null
+++ b/Engine/Plugins/SubtitlesAndClosedCaptions/Source/SubtitlesAndClosedCaptions/Private/SubtitlesSubsystem.cpp
@@ -0,0 +1,373 @@
+// Copyright Epic Games, Inc. All Rights Reserved.
+
+#include "SubtitlesSubsystem.h"
+
+#include "Algo/Find.h"
+#include "Algo/RemoveIf.h"
+#include "CanvasItem.h"
+#include "CanvasTypes.h"
+#include "Containers/Ticker.h"
+#include "Engine/Font.h"
+#include "Engine/World.h"
+#include "Engine/Engine.h"
+#include "Logging/StructuredLog.h"
+#include "Math/UnrealMathUtility.h"
+#include "Styling/CoreStyle.h"
+#include "SubtitlesAndClosedCaptionsModule.h"
+#include "SubtitlesSettings.h"
+#include "Subtitles/SubtitlesAndClosedCaptionsDelegates.h"
+#include "TimerManager.h"
+#include "UnrealClient.h"
+#include "UObject/ObjectMacros.h"
+#include "UObject/Package.h"
+#include "UObject/StrongObjectPtrTemplates.h"
+#include "UObject/WeakObjectPtrTemplates.h"
+
+static constexpr float InfiniteDuration = FLT_MAX;
+
+#include UE_INLINE_GENERATED_CPP_BY_NAME(SubtitlesSubsystem)
+
+void USubtitlesSubsystem::Initialize(FSubsystemCollectionBase& Collection)
+{
+	BindDelegates();
+}
+
+void USubtitlesSubsystem::BindDelegates()
+{
+	if (TryCreateUMGWidget())
+	{
+		FSubtitlesAndClosedCaptionsDelegates::QueueSubtitle.BindUObject(this, &USubtitlesSubsystem::QueueSubtitle);
+		FSubtitlesAndClosedCaptionsDelegates::IsSubtitleActive.BindUObject(this, &USubtitlesSubsystem::IsSubtitleActive);
+		FSubtitlesAndClosedCaptionsDelegates::StopSubtitle.BindUObject(this, &USubtitlesSubsystem::StopSubtitle);
+		FSubtitlesAndClosedCaptionsDelegates::StopAllSubtitles.BindUObject(this, &USubtitlesSubsystem::StopAllSubtitles);
+	}
+}
+
+void USubtitlesSubsystem::QueueSubtitle(FQueueSubtitleParameters Params, const ESubtitleTiming Timing)
+{
+	// Don't queue subtitles when the engine doesn't have them enabled or forced off. 
+	if (GEngine == nullptr)
+	{
+		UE_LOGFMT(LogSubtitlesAndClosedCaptions, Warning, "Couldn't check if subtitles are enabled, because GEngine nullptr.");
+		return;
+	}
+	if (!GEngine->bSubtitlesEnabled || GEngine->bSubtitlesForcedOff)
+	{
+		return;
+	}
+
+	const FSubtitleAssetData& Subtitle = Params.Subtitle;
+	UE_LOGFMT(LogSubtitlesAndClosedCaptions, Display, "QueueSubtitle: {Name}", Subtitle.Text.ToString());
+
+	// Externally-timed subtitles will be removed by the system queueing them, so they have an otherwise-infinite duration.
+	float Duration = InfiniteDuration;
+	if (Timing == ESubtitleTiming::InternallyTimed)
+	{
+		Duration = Params.Duration.IsSet() ? Params.Duration.GetValue() : Subtitle.Duration;
+	}
+
+	const float StartOffset = Params.StartOffset.IsSet() ? Params.StartOffset.GetValue() : Subtitle.StartOffset;
+
+	if (IsInGameThread())
+	{
+		AddActiveSubtitle(Subtitle, Duration, StartOffset, Timing);
+	}
+	else
+	{
+		ExecuteOnGameThread(
+			TEXT("USubtitlesSubsystem::HandleQueueSubtitle"),
+			[WeakThis = TWeakObjectPtr<USubtitlesSubsystem>(this), Subtitle, Duration, StartOffset, Timing]()
+			{
+				if (TStrongObjectPtr<USubtitlesSubsystem> ThisPin = WeakThis.Pin())
+				{
+					ThisPin->AddActiveSubtitle(Subtitle, Duration, StartOffset, Timing);
+				}
+			}
+		);
+	}
+}
+
+void USubtitlesSubsystem::AddActiveSubtitle(const FSubtitleAssetData& Subtitle, float Duration, const float StartOffset, const ESubtitleTiming Timing)
+{
+	// If the subtitle is already active then update its duration (by removing it then and then re-adding it)
+	const FActiveSubtitleAssetData* FoundActiveSubtitle = ActiveSubtitles.FindByPredicate(
+		[&Subtitle](const FActiveSubtitleAssetData& ActiveSubtitle) { return ActiveSubtitle.Subtitle == Subtitle; });
+	if (FoundActiveSubtitle != nullptr)
+	{
+		RemoveActiveSubtitle(Subtitle);
+	}
+
+	FActiveSubtitleAssetData NewActiveSubtitle{ Subtitle };
+
+	// Subtitles with delayed offset need a timer to await their entry in the queue.
+	if (StartOffset > 0.f)
+	{
+		// The timer handle will be reused for the duration when it enters the active subtitle queue.
+		// For now the timer tracks how long until it enters that queue.
+		FTimerDelegate Delegate;
+		Delegate.BindUFunction(this, FName(TEXT("MakeDelayedSubtitleActive")), NewActiveSubtitle.Subtitle, Timing);
+		GetWorldRef().GetTimerManager().SetTimer(NewActiveSubtitle.DurationTimerHandle, MoveTemp(Delegate), StartOffset, /*bLoop=*/false);
+
+		DelayedSubtitles.Add(MoveTemp(NewActiveSubtitle));
+	}
+	else
+	{
+		if (Timing == ESubtitleTiming::InternallyTimed)
+		{
+			// Without the delayed offset, instantly enter the queue as usual.
+			// The timer here tracks how long until the subtitle will expire and leave the active subtitle queue.
+			FTimerDelegate Delegate;
+			Delegate.BindUFunction(this, FName(TEXT("RemoveActiveSubtitle")), NewActiveSubtitle.Subtitle);
+			Duration = FMath::Max(Duration, SubtitleMinDuration);
+			GetWorldRef().GetTimerManager().SetTimer(NewActiveSubtitle.DurationTimerHandle, MoveTemp(Delegate), Duration, /*bLoop=*/false);
+		}
+		
+		AddAndDisplaySubtitle(NewActiveSubtitle);
+	}
+}
+
+void USubtitlesSubsystem::MakeDelayedSubtitleActive(const FSubtitleAssetData& Subtitle, const ESubtitleTiming Timing)
+{
+	FActiveSubtitleAssetData* DelayedSubtitle = Algo::FindByPredicate(DelayedSubtitles, [&Subtitle](const FActiveSubtitleAssetData& DelayedSubtitle) {return DelayedSubtitle.Subtitle == Subtitle; });
+
+	if (DelayedSubtitle != nullptr)
+	{
+		const FSubtitleAssetData& SubtitleAsset = DelayedSubtitle->Subtitle;
+		const float Duration = FMath::Max(SubtitleAsset.Duration, SubtitleMinDuration);
+
+		if (Timing == ESubtitleTiming::InternallyTimed)
+		{
+			// Reuse the Timer Handle for duration, now that it's no longer needed for the delay.
+			FTimerDelegate Delegate;
+			Delegate.BindUFunction(this, FName(TEXT("RemoveActiveSubtitle")), DelayedSubtitle->Subtitle);
+			GetWorldRef().GetTimerManager().SetTimer(DelayedSubtitle->DurationTimerHandle, MoveTemp(Delegate), Duration, /*bLoop=*/false);
+		}
+
+		// Insert the new subtitle to the actual queue and ensure it remains sorted by priority.
+		AddAndDisplaySubtitle(*DelayedSubtitle);
+
+		// Remove from the list of Delayed Subtitles
+		const int32 FirstRemovedIndex = Algo::RemoveIf(DelayedSubtitles, [&Subtitle](const FActiveSubtitleAssetData& DelayedSubtitle) {return DelayedSubtitle.Subtitle == Subtitle; });
+		DelayedSubtitles.SetNum(FirstRemovedIndex);
+	}
+}
+
+void USubtitlesSubsystem::AddAndDisplaySubtitle(FActiveSubtitleAssetData& NewActiveSubtitle)
+{
+	const FSubtitleAssetData& NewSubtitle = NewActiveSubtitle.Subtitle;
+	const float NewPriority = NewSubtitle.Priority;
+	const ESubtitleType Category = NewSubtitle.SubtitleType;
+
+	bool ShouldAddSubtitle = true;
+	// Warn if the subtitle won't play due to being lower-priority than the currently-playing subtitle in that category.
+	for (const FActiveSubtitleAssetData& PreQueuedSubtitle : ActiveSubtitles)
+	{
+		const FSubtitleAssetData& PreQueuedSubtitleData = PreQueuedSubtitle.Subtitle;
+		if (PreQueuedSubtitleData.SubtitleType == Category && (PreQueuedSubtitleData.Priority > NewPriority))
+		{
+			UE_LOGFMT(LogSubtitlesAndClosedCaptions, Display, "Subtitle {Name} won't display: its priority is lower than the currently-playing subtitle in that category.", NewSubtitle.Text.ToString());
+			ShouldAddSubtitle = false;
+			break;
+		}
+	}
+
+	// Remove newly lower-priority subtitles.
+	FTimerManager& TimerManager = GetWorldRef().GetTimerManager();
+
+	const int32 FirstRemovedIndex = Algo::StableRemoveIf(ActiveSubtitles,
+		[NewPriority, Category](const FActiveSubtitleAssetData& RemovalCandidate) {return (RemovalCandidate.Subtitle.Priority < NewPriority) && RemovalCandidate.Subtitle.SubtitleType == Category; }
+	);
+	for (int32 Index = FirstRemovedIndex; Index < ActiveSubtitles.Num(); ++Index)
+	{
+		TimerManager.ClearTimer(ActiveSubtitles[Index].DurationTimerHandle);
+	}
+
+	ActiveSubtitles.SetNum(FirstRemovedIndex);
+
+	if (ShouldAddSubtitle)
+	{
+		ActiveSubtitles.Add(MoveTemp(NewActiveSubtitle));
+
+		// Sort by priority, but also prefer newer subtitles with a delayed start.
+		ActiveSubtitles.StableSort([](const FActiveSubtitleAssetData& Lhs, const FActiveSubtitleAssetData& Rhs)
+			{ return (Lhs.Subtitle.Priority > Rhs.Subtitle.Priority) || (Lhs.Subtitle.StartOffset > 0 && (Lhs.Subtitle.Priority == Rhs.Subtitle.Priority)); });
+
+		UpdateWidgetData();
+	}
+}
+
+bool USubtitlesSubsystem::IsSubtitleActive(FSubtitleAssetData Data) const
+{
+	if (!ensureMsgf(IsInGameThread(), TEXT("IsSubtitleActive must currently be run on the GameThread - ActiveSubtitles vector is not locked")))
+	{
+		return false;
+	}
+
+	const FActiveSubtitleAssetData* FoundActiveSubtitle = ActiveSubtitles.FindByPredicate(
+		[Data](const FActiveSubtitleAssetData& ActiveSubtitle) { return ActiveSubtitle.Subtitle == Data; });
+
+	return FoundActiveSubtitle != nullptr;
+}
+
+void USubtitlesSubsystem::StopSubtitle(FSubtitleAssetData Data)
+{
+	RemoveActiveSubtitle(Data);
+}
+
+void USubtitlesSubsystem::StopAllSubtitles()
+{
+	// Clean up queued subtitles.
+	FTimerManager& TimerManager = GetWorldRef().GetTimerManager();
+	for (FActiveSubtitleAssetData& ActiveSubtitle : ActiveSubtitles)
+	{
+		TimerManager.ClearTimer(ActiveSubtitle.DurationTimerHandle);
+	}
+	ActiveSubtitles.Empty();
+
+	// Also remove delayed-start subtitles not yet in the queue.
+	for (FActiveSubtitleAssetData& DelayedSubtitle : DelayedSubtitles)
+	{
+		TimerManager.ClearTimer(DelayedSubtitle.DurationTimerHandle);
+	}
+	DelayedSubtitles.Empty();
+
+	// Clear the widget's display
+	if (IsValid(SubtitleWidget))
+	{
+		SubtitleWidget->StopDisplayingAllSubtitles();
+	}
+	else 
+	{
+		UE_LOGFMT(LogSubtitlesAndClosedCaptions, Warning, "Can't remove subtitles because there isn't a valid UMG widget.");
+	}
+}
+
+void USubtitlesSubsystem::ReplaceWidget(const TSoftClassPtr<USubtitleWidget>& NewWidgetAsset)
+{
+	if (!ActiveSubtitles.IsEmpty() || !DelayedSubtitles.IsEmpty())
+	{
+		UE_LOGFMT(LogSubtitlesAndClosedCaptions, Error, "Can't replace the subtitle widget: there are still subtitles queued or displaying. Use StopAllSubtitles() first.");
+		return;
+	}
+
+	if (!TryCreateUMGWidgetFromAsset(NewWidgetAsset))
+	{
+		UE_LOGFMT(LogSubtitlesAndClosedCaptions, Error, "Can't replace the subtitle widget; was a valid asset provided?");
+	}
+}
+
+void USubtitlesSubsystem::RemoveActiveSubtitle(const FSubtitleAssetData& Subtitle)
+{
+	bool bSuccessfullyRemoved = false;
+
+	FTimerManager& TimerManager = GetWorldRef().GetTimerManager();
+
+	// Remove currently-queued subtitles
+	const int32 FirstRemovedIndex = Algo::StableRemoveIf(ActiveSubtitles, 
+		[&Subtitle](const FActiveSubtitleAssetData& ActiveSubtitle) {return ActiveSubtitle.Subtitle == Subtitle; }
+	);
+	for (int32 Index = FirstRemovedIndex; Index < ActiveSubtitles.Num(); ++Index)
+	{
+		TimerManager.ClearTimer(ActiveSubtitles[Index].DurationTimerHandle);
+		bSuccessfullyRemoved = true;
+	}
+
+	ActiveSubtitles.SetNum(FirstRemovedIndex);
+
+	// Stop Displaying the removed subtitle and display a newly-most-relevant one if applicable.
+	if (bSuccessfullyRemoved && IsValid(SubtitleWidget))
+	{
+		SubtitleWidget->StopDisplayingSubtitle(Subtitle);
+
+		if (!ActiveSubtitles.IsEmpty())
+		{
+			const FSubtitleAssetData& HighestPrioritySubtitle = ActiveSubtitles[0].Subtitle;
+			SubtitleWidget->StartDisplayingSubtitle(HighestPrioritySubtitle);
+		}
+	}
+	else if (!IsValid(SubtitleWidget))
+	{
+		UE_LOGFMT(LogSubtitlesAndClosedCaptions, Warning, "Can't remove subtitles because there isn't a valid UMG widget.");
+	}
+
+	// Remove delayed-start subtitles that haven't been queued yet that also use this asset.
+	const int32 FirstRemovedActiveIndex = Algo::RemoveIf(DelayedSubtitles,
+		[&Subtitle](const FActiveSubtitleAssetData& DelayedSubtitle) {return DelayedSubtitle.Subtitle == Subtitle; }
+	);
+
+	for (int32 Index = FirstRemovedActiveIndex; Index < DelayedSubtitles.Num(); ++Index)
+	{
+		TimerManager.ClearTimer(DelayedSubtitles[Index].DurationTimerHandle);
+	}
+	DelayedSubtitles.SetNum(FirstRemovedActiveIndex);
+}
+
+bool USubtitlesSubsystem::TryCreateUMGWidgetFromAsset(const TSoftClassPtr<USubtitleWidget>& WidgetToUse)
+{
+	// Synchronous load: called directly from ::Initialize when the World (and this subsystem) is created.
+	WidgetToUse.LoadSynchronous();
+	if (WidgetToUse.IsValid())
+	{
+		SubtitleWidget = CreateWidget<USubtitleWidget>(GetWorld(), WidgetToUse.Get());
+		if (IsValid(SubtitleWidget))
+		{
+			SubtitleWidget->SetFlags(RF_Transient);
+		}
+		bInitializedWidget = false;	// The new widget will need to be added to the viewport.
+	}
+
+	return IsValid(SubtitleWidget);
+}
+
+bool USubtitlesSubsystem::TryCreateUMGWidget()
+{
+	// Set up the UMG widget
+	const USubtitlesSettings* Settings = GetDefault<USubtitlesSettings>();
+	check(Settings != nullptr);
+	const TSoftClassPtr<USubtitleWidget>& WidgetToUse = Settings->GetWidgetSoftClassPtr();
+	
+	if (!TryCreateUMGWidgetFromAsset(WidgetToUse))
+	{
+		// Fallback to default widget (not set by user):
+		const TSoftClassPtr<USubtitleWidget>& WidgetToUseDefault = Settings->GetWidgetDefaultSoftClassPtr();
+		if (WidgetToUseDefault.IsValid())
+		{
+			SubtitleWidget = CreateWidget<USubtitleWidget>(GetWorld(), WidgetToUseDefault.Get());
+			if (IsValid(SubtitleWidget))
+			{
+				SubtitleWidget->SetFlags(RF_Transient);
+			}
+			// Log the fallback as a warning as seeing the default widget probably isn't intended.
+			UE_LOGFMT(LogSubtitlesAndClosedCaptions, Warning, "A valid Subtitle Widget wasn't provided. The plugin's default widget will be used as a fallback.");
+		}
+		else
+		{
+			UE_LOGFMT(LogSubtitlesAndClosedCaptions, Error, "The default Subtitle Widget asset isn't valid. Subtitles won't be displayed.");
+		}
+	}
+
+	bInitializedWidget = false;
+	return IsValid(SubtitleWidget);
+}
+
+void USubtitlesSubsystem::UpdateWidgetData()
+{
+	// Update the widget. If it's not valid (eg, destroyed on non-seamless travel), try re-creating it first.
+	if (IsValid(SubtitleWidget) || (bInitializedWidget && TryCreateUMGWidget() && IsValid(SubtitleWidget)))
+	{
+		if (!bInitializedWidget)
+		{
+			SubtitleWidget->AddToViewport();
+			SubtitleWidget->SetVisibility(ESlateVisibility::HitTestInvisible);
+			bInitializedWidget = true;
+		}
+
+		const FSubtitleAssetData& NewHighestPrioritySubtitle = ActiveSubtitles[0].Subtitle;
+		SubtitleWidget->StartDisplayingSubtitle(NewHighestPrioritySubtitle);
+	}
+	else
+	{
+		UE_LOGFMT(LogSubtitlesAndClosedCaptions, Warning, "Can't display subtitles because there isn't a valid UMG widget to display it to (check Project Settings).");
+	}
+}
+
diff --git a/Engine/Plugins/SubtitlesAndClosedCaptions/Source/SubtitlesAndClosedCaptions/Public/ActiveSubtitleAssetData.h b/Engine/Plugins/SubtitlesAndClosedCaptions/Source/SubtitlesAndClosedCaptions/Public/ActiveSubtitleAssetData.h
new file mode 100644
--- /dev/null
+++ b/Engine/Plugins/SubtitlesAndClosedCaptions/Source/SubtitlesAndClosedCaptions/Public/ActiveSubtitleAssetData.h
@@ -0,0 +1,22 @@
+// Copyright Epic Games, Inc. All Rights Reserved.
+
+#pragma once 
+
+struct FSubtitleAssetData;
+
+#include "Engine/TimerHandle.h"
+#include "Subtitles/SubtitlesAndClosedCaptionsTypes.h"
+#include "UObject/ObjectPtr.h"
+
+#include "ActiveSubtitleAssetData.generated.h"
+
+USTRUCT()
+struct FActiveSubtitleAssetData
+{
+	GENERATED_BODY()
+
+	UPROPERTY()
+	FSubtitleAssetData Subtitle;
+
+	FTimerHandle DurationTimerHandle;
+};
diff --git a/Engine/Plugins/SubtitlesAndClosedCaptions/Source/SubtitlesAndClosedCaptions/Public/MovieSceneSubtitleComponentTypes.h b/Engine/Plugins/SubtitlesAndClosedCaptions/Source/SubtitlesAndClosedCaptions/Public/MovieSceneSubtitleComponentTypes.h
new file mode 100644
--- /dev/null
+++ b/Engine/Plugins/SubtitlesAndClosedCaptions/Source/SubtitlesAndClosedCaptions/Public/MovieSceneSubtitleComponentTypes.h
@@ -0,0 +1,26 @@
+// Copyright Epic Games, Inc. All Rights Reserved.
+
+#pragma once
+
+#include "EntitySystem/MovieSceneEntityIDs.h"
+#include "MovieSceneFwd.h"
+#include "SubtitleDataComponent.h"
+
+struct FSubtitleDataComponent;
+
+class UMovieSceneSection;
+
+// Component data present on all Subtitle entities
+struct FMovieSceneSubtitleComponentTypes
+{
+	UE::MovieScene::TComponentTypeID<FSubtitleDataComponent> SubtitleData;
+
+	SUBTITLESANDCLOSEDCAPTIONS_API static FMovieSceneSubtitleComponentTypes* Get();
+	static void Destroy();
+
+private:
+	FMovieSceneSubtitleComponentTypes();
+
+	static bool GMovieSceneSubtitlesComponentTypesDestroyed;
+	static TUniquePtr<FMovieSceneSubtitleComponentTypes> GMovieSceneSubtitlesComponentTypes;
+};
diff --git a/Engine/Plugins/SubtitlesAndClosedCaptions/Source/SubtitlesAndClosedCaptions/Public/MovieSceneSubtitleSection.h b/Engine/Plugins/SubtitlesAndClosedCaptions/Source/SubtitlesAndClosedCaptions/Public/MovieSceneSubtitleSection.h
new file mode 100644
--- /dev/null
+++ b/Engine/Plugins/SubtitlesAndClosedCaptions/Source/SubtitlesAndClosedCaptions/Public/MovieSceneSubtitleSection.h
@@ -0,0 +1,46 @@
+// Copyright Epic Games, Inc. All Rights Reserved.
+
+#pragma once
+
+#include "EntitySystem/IMovieSceneEntityProvider.h"
+#include "EntitySystem/MovieSceneEntityIDs.h"
+#include "MovieSceneSection.h"
+#include "UObject/ObjectMacros.h"
+
+#include "MovieSceneSubtitleSection.generated.h"
+
+class USubtitleAssetUserData;
+struct FMovieSceneEntityComponentFieldBuilder;
+struct FMovieSceneEvaluationFieldEntityMetaData;
+
+UCLASS(MinimalAPI)
+class UMovieSceneSubtitleSection
+	: public UMovieSceneSection
+	, public IMovieSceneEntityProvider
+{
+	GENERATED_UCLASS_BODY()
+
+public:
+
+	TObjectPtr<const USubtitleAssetUserData> GetSubtitle() const
+	{
+		return Subtitle;
+	}
+
+	void SetSubtitle(const USubtitleAssetUserData& InSubtitle)
+	{
+		Subtitle = &InSubtitle;
+	}
+
+	//~ UMovieSceneSection interface
+	SUBTITLESANDCLOSEDCAPTIONS_API virtual TOptional<TRange<FFrameNumber> > GetAutoSizeRange() const override;
+
+	//~ IMovieSceneEntityProvider interface
+	SUBTITLESANDCLOSEDCAPTIONS_API virtual bool PopulateEvaluationFieldImpl(const TRange<FFrameNumber>& EffectiveRange, const FMovieSceneEvaluationFieldEntityMetaData& InMetaData, FMovieSceneEntityComponentFieldBuilder* OutFieldBuilder) override;
+	SUBTITLESANDCLOSEDCAPTIONS_API virtual void ImportEntityImpl(UMovieSceneEntitySystemLinker* EntityLinker, const FEntityImportParams& Params, FImportedEntity* OutImportedEntity) override;
+
+private:
+
+	UPROPERTY(EditAnywhere, Category = "Subtitles")
+	TObjectPtr<const USubtitleAssetUserData> Subtitle;
+};
diff --git a/Engine/Plugins/SubtitlesAndClosedCaptions/Source/SubtitlesAndClosedCaptions/Public/MovieSceneSubtitlesSystem.h b/Engine/Plugins/SubtitlesAndClosedCaptions/Source/SubtitlesAndClosedCaptions/Public/MovieSceneSubtitlesSystem.h
new file mode 100644
--- /dev/null
+++ b/Engine/Plugins/SubtitlesAndClosedCaptions/Source/SubtitlesAndClosedCaptions/Public/MovieSceneSubtitlesSystem.h
@@ -0,0 +1,152 @@
+// Copyright Epic Games, Inc. All Rights Reserved.
+
+#pragma once
+
+#include "CoreTypes.h"
+#include "EntitySystem/MovieSceneEntityIDs.h"
+#include "EntitySystem/MovieSceneEntitySystem.h"
+#include "EntitySystem/MovieSceneSequenceInstanceHandle.h"
+#include "MovieSceneTracksComponentTypes.h"
+#include "UObject/ObjectMacros.h"
+
+#include "EntitySystem/MovieSceneEntitySystemLinker.h"
+#include "EntitySystem/MovieSceneEntitySystemTask.h"
+#include "MovieSceneSubtitleSection.h"
+#include "SubtitleDataComponent.h"
+#include "Subtitles/SubtitlesAndClosedCaptionsDelegates.h"
+#include "Subtitles/SubtitlesAndClosedCaptionsTypes.h"
+
+#include "MovieSceneSubtitlesSystem.generated.h"
+
+struct FEvaluateSubtitles
+{
+
+	FEvaluateSubtitles(const UE::MovieScene::FInstanceRegistry& InInstanceRegistry)
+		: InstanceRegistry(InInstanceRegistry)
+	{
+	}
+
+	void ForEachAllocation(
+		const UE::MovieScene::FEntityAllocation* Allocation,
+		UE::MovieScene::TRead<UE::MovieScene::FInstanceHandle> SequenceInstanceHandles,
+		UE::MovieScene::TWrite<FSubtitleDataComponent> SubtitleData) const
+	{
+		const UE::MovieScene::FBuiltInComponentTypes* BuiltInComponents = UE::MovieScene::FBuiltInComponentTypes::Get();
+
+		const int32 Num = Allocation->Num();
+		for (int32 Index = 0; Index < Num; ++Index)
+		{
+			const UE::MovieScene::FInstanceHandle& SequenceInstanceHandle = SequenceInstanceHandles[Index];
+			const UE::MovieScene::FSequenceInstance& SequenceInstance = InstanceRegistry.GetInstance(SequenceInstanceHandle);
+
+			// Evaluate the subtitle associated with each subtitle section with logic dependent on changes in the context playing state
+			Evaluate(SequenceInstance,
+				SubtitleData[Index],
+				SubtitleData[Index].LastSequenceInstanceStatus);
+
+			SubtitleData[Index].LastSequenceInstanceStatus = SequenceInstance.GetContext().GetStatus();
+		}
+	}
+
+	void Evaluate(const UE::MovieScene::FSequenceInstance& SequenceInstance, const FSubtitleDataComponent& SubtitleData, EMovieScenePlayerStatus::Type LastStatus) const
+	{
+		const FMovieSceneContext& Context = SequenceInstance.GetContext();
+
+		const UMovieSceneSubtitleSection* SubtitleSection = SubtitleData.SubtitleSection;
+		if (!ensureMsgf(SubtitleSection, TEXT("No valid subtitle section found in subtitles data!")))
+		{
+			return;
+		}
+
+		const USubtitleAssetUserData* SubtitleAsset = SubtitleData.SubtitleSection->GetSubtitle();
+		if (!ensureMsgf(SubtitleAsset, TEXT("No valid subtitle found in subtitle section!")))
+		{
+			return;
+		}
+
+		const EMovieScenePlayerStatus::Type Status = Context.GetStatus();
+#if !NO_LOGGING
+		if (const UEnum* Enum = StaticEnum<EMovieScenePlayerStatus::Type>())
+		{
+			/* TODO: enable for ClangEditor Win64
+			const FString LastStatusString = Enum->GetNameStringByIndex(LastStatus);
+			const FString StatusString = Enum->GetNameStringByIndex(Status);
+			UE_LOGFMT(LogSubtitlesAndClosedCaptions, Verbose, "LastStatus: {LastStatus}, Status: {CurStatus}, HasJumped: {Jumped}",
+				LastStatusString, StatusString, Context.HasJumped());
+			*/
+		}
+#endif // !NO_LOGGING
+
+		// Stop all subtitles when we start playing to clear out any infinite duration subtitiles that were queued on pause/stop
+		const bool bHasStarted{ (Status == EMovieScenePlayerStatus::Type::Playing || Status == EMovieScenePlayerStatus::Type::Scrubbing)
+		&& (LastStatus == EMovieScenePlayerStatus::Type::Stopped || LastStatus == EMovieScenePlayerStatus::Type::Paused || LastStatus == EMovieScenePlayerStatus::Type::Stepping) };
+		if (bHasStarted || Context.HasJumped())
+		{
+			FSubtitlesAndClosedCaptionsDelegates::StopAllSubtitles.ExecuteIfBound();
+		}
+
+		if (FSubtitlesAndClosedCaptionsDelegates::IsSubtitleActive.IsBound())
+		{
+			const FFrameNumber LastSubtitleFrame = SubtitleSection->GetRange().GetUpperBoundValue();
+			const FFrameNumber FirstSubtitleFrame = SubtitleSection->GetRange().GetLowerBoundValue();
+			const FFrameNumber LastEvaluatedFrame = Context.GetFrameNumberRange().GetUpperBoundValue();
+
+			// Only try to play subtitles if inside the section's first and last frames.
+			if (LastEvaluatedFrame >= FirstSubtitleFrame && LastEvaluatedFrame < LastSubtitleFrame)
+			{
+				const FFrameRate FrameRate = Context.GetFrameRate();
+				for (const FSubtitleAssetData& Entry : SubtitleAsset->Subtitles)
+				{
+					// The subtitle asset may have several subtitles.
+					// To ensure each is queued at the correct time, check the start and end times of each against the current playback position.
+					const FFrameTime StartOffset = Entry.StartOffset * FrameRate;
+					const FFrameNumber FirstOffsetSubtitleFrame = FirstSubtitleFrame + StartOffset.FrameNumber;
+
+					const FFrameTime LastFrameTime = Entry.Duration * FrameRate;
+					const FFrameNumber LastOffsetSubtitleFrame = FirstOffsetSubtitleFrame + LastFrameTime.FrameNumber;
+
+					if (LastEvaluatedFrame >= FirstOffsetSubtitleFrame && LastEvaluatedFrame < LastOffsetSubtitleFrame)
+					{
+						// Queue the subtitle to display in short (1/10th second) increments whenever playback is over the subtitle.
+						// This allows for jumping and scrubbing the sequence without the subtitle getting stuck on or off.
+						FQueueSubtitleParameters Params{ Entry };
+						Params.Duration = 0.1f;
+						Params.StartOffset = 0.f;
+
+						FSubtitlesAndClosedCaptionsDelegates::QueueSubtitle.ExecuteIfBound(Params, ESubtitleTiming::InternallyTimed);
+					}
+				}
+			}
+			else if(LastSubtitleFrame <= LastEvaluatedFrame)
+			{
+				for (const FSubtitleAssetData& Entry : SubtitleAsset->Subtitles)
+				{
+					if (FSubtitlesAndClosedCaptionsDelegates::IsSubtitleActive.Execute(Entry))
+					{
+						// If the subtitle is already active and outside the range of frames it should display for, remove it.
+						FSubtitlesAndClosedCaptionsDelegates::StopSubtitle.ExecuteIfBound(Entry);
+					}
+				}
+			}
+		}
+	}
+
+protected:
+	const UE::MovieScene::FInstanceRegistry& InstanceRegistry;
+};
+
+UCLASS()
+class UMovieSceneSubtitlesSystem : public UMovieSceneEntitySystem
+{
+	GENERATED_BODY()
+
+public:
+
+	UMovieSceneSubtitlesSystem(const FObjectInitializer& ObjInit);
+	virtual ~UMovieSceneSubtitlesSystem();
+
+	//~ UMovieSceneEntitySystem members
+	virtual void OnSchedulePersistentTasks(UE::MovieScene::IEntitySystemScheduler* TaskScheduler) override;
+	virtual void OnRun(FSystemTaskPrerequisites& InPrerequisites, FSystemSubsequentTasks& Subsequents) override;
+};
+
diff --git a/Engine/Plugins/SubtitlesAndClosedCaptions/Source/SubtitlesAndClosedCaptions/Public/MovieSceneSubtitlesTrack.h b/Engine/Plugins/SubtitlesAndClosedCaptions/Source/SubtitlesAndClosedCaptions/Public/MovieSceneSubtitlesTrack.h
new file mode 100644
--- /dev/null
+++ b/Engine/Plugins/SubtitlesAndClosedCaptions/Source/SubtitlesAndClosedCaptions/Public/MovieSceneSubtitlesTrack.h
@@ -0,0 +1,41 @@
+// Copyright Epic Games, Inc. All Rights Reserved.
+
+#pragma once
+
+#include "UObject/ObjectMacros.h"
+#include "MovieSceneNameableTrack.h"
+
+#include "MovieSceneSubtitlesTrack.generated.h"
+
+class UMovieSceneSection;
+class USubtitleAssetUserData;
+
+UCLASS(MinimalAPI)
+class UMovieSceneSubtitlesTrack
+	: public UMovieSceneNameableTrack
+{
+	GENERATED_UCLASS_BODY()
+
+public:
+	SUBTITLESANDCLOSEDCAPTIONS_API UMovieSceneSection* AddNewSubtitle(const USubtitleAssetUserData& Subtitle, FFrameNumber InStartTime);
+	SUBTITLESANDCLOSEDCAPTIONS_API UMovieSceneSection* AddNewSubtitleOnRow(const USubtitleAssetUserData& Subtitle, FFrameNumber InStartTime, int32 InRowIndex);
+
+public:
+
+	// UMovieSceneTrack interface
+	SUBTITLESANDCLOSEDCAPTIONS_API virtual bool SupportsType(TSubclassOf<UMovieSceneSection> SectionClass) const override;
+	SUBTITLESANDCLOSEDCAPTIONS_API virtual void RemoveAllAnimationData() override;
+	SUBTITLESANDCLOSEDCAPTIONS_API virtual bool HasSection(const UMovieSceneSection& Section) const override;
+	SUBTITLESANDCLOSEDCAPTIONS_API virtual void AddSection(UMovieSceneSection& Section) override;
+	SUBTITLESANDCLOSEDCAPTIONS_API virtual void RemoveSection(UMovieSceneSection& Section) override;
+	SUBTITLESANDCLOSEDCAPTIONS_API virtual void RemoveSectionAt(int32 SectionIndex) override;
+	SUBTITLESANDCLOSEDCAPTIONS_API virtual bool IsEmpty() const override;
+	SUBTITLESANDCLOSEDCAPTIONS_API virtual const TArray<UMovieSceneSection*>& GetAllSections() const override;
+	SUBTITLESANDCLOSEDCAPTIONS_API virtual bool SupportsMultipleRows() const override;
+	SUBTITLESANDCLOSEDCAPTIONS_API virtual UMovieSceneSection* CreateNewSection() override;
+
+private:
+	UPROPERTY()
+	TArray<TObjectPtr<UMovieSceneSection>> SubtitleSections;
+
+};
diff --git a/Engine/Plugins/SubtitlesAndClosedCaptions/Source/SubtitlesAndClosedCaptions/Public/SubtitleDataComponent.h b/Engine/Plugins/SubtitlesAndClosedCaptions/Source/SubtitlesAndClosedCaptions/Public/SubtitleDataComponent.h
new file mode 100644
--- /dev/null
+++ b/Engine/Plugins/SubtitlesAndClosedCaptions/Source/SubtitlesAndClosedCaptions/Public/SubtitleDataComponent.h
@@ -0,0 +1,23 @@
+// Copyright Epic Games, Inc. All Rights Reserved.
+
+#pragma once
+
+#include "EntitySystem/MovieSceneEntityIDs.h"
+#include "MovieSceneFwd.h"
+#include "UObject/ObjectPtr.h"
+
+#include "SubtitleDataComponent.generated.h"
+
+class UMovieSceneSubtitleSection;
+
+/** Component data for subtitles tracks */
+USTRUCT()
+struct FSubtitleDataComponent
+{
+	GENERATED_BODY()
+
+	UPROPERTY()
+	TObjectPtr<const UMovieSceneSubtitleSection> SubtitleSection;
+
+	EMovieScenePlayerStatus::Type LastSequenceInstanceStatus = EMovieScenePlayerStatus::Type::Stopped;
+};
diff --git a/Engine/Plugins/SubtitlesAndClosedCaptions/Source/SubtitlesAndClosedCaptions/Public/SubtitleWidget.h b/Engine/Plugins/SubtitlesAndClosedCaptions/Source/SubtitlesAndClosedCaptions/Public/SubtitleWidget.h
new file mode 100644
--- /dev/null
+++ b/Engine/Plugins/SubtitlesAndClosedCaptions/Source/SubtitlesAndClosedCaptions/Public/SubtitleWidget.h
@@ -0,0 +1,52 @@
+// Copyright Epic Games, Inc. All Rights Reserved.
+
+#pragma once
+#include "Blueprint/UserWidget.h"
+#include "Subtitles/SubtitlesAndClosedCaptionsDelegates.h"
+
+#include "SubtitleWidget.generated.h"
+
+#define UE_API SUBTITLESANDCLOSEDCAPTIONS_API
+
+class UBorder;
+class UTextBlock;
+
+UCLASS(Abstract, BlueprintType, Blueprintable, meta=( DontUseGenericSpawnObject="True", DisableNativeTick), MinimalAPI)
+class USubtitleWidget : public UUserWidget
+{
+	GENERATED_BODY()
+
+private:
+	TMap<ESubtitleType, TOptional<FSubtitleAssetData>> CurrentlyDisplayedSubtitles;
+
+protected:
+	UE_API virtual void NativeConstruct() override;
+
+	// Subtitle Text Blocks
+	UPROPERTY(BlueprintReadOnly, Category = "Subtitles", meta = (BindWidgetOptional))
+	TObjectPtr<UTextBlock> DialogSubtitleBlock;
+
+	UPROPERTY(BlueprintReadOnly, Category = "Subtitles", meta = (BindWidgetOptional))
+	TObjectPtr<UTextBlock> CaptionSubtitleBlock;
+
+	UPROPERTY(BlueprintReadOnly, Category = "Subtitles", meta = (BindWidgetOptional))
+	TObjectPtr<UTextBlock> DescriptionSubtitleBlock;
+
+	// Borders to the text blocks
+	UPROPERTY(BlueprintReadOnly, Category = "SubtitleBorders", meta = (BindWidgetOptional))
+	TObjectPtr<UBorder> DialogBorder;
+
+	UPROPERTY(BlueprintReadOnly, Category = "SubtitleBorders", meta = (BindWidgetOptional))
+	TObjectPtr<UBorder> CaptionBorder;
+
+	UPROPERTY(BlueprintReadOnly, Category = "SubtitleBorders", meta = (BindWidgetOptional))
+	TObjectPtr<UBorder> DescriptionBorder;
+public:
+	void StartDisplayingSubtitle(const FSubtitleAssetData& Subtitle);
+	void StopDisplayingSubtitle(const FSubtitleAssetData& Subtitle);
+	void StopDisplayingAllSubtitles();
+
+	UE_API USubtitleWidget(const FObjectInitializer& ObjectInitializer);
+};
+
+#undef UE_API
\ No newline at end of file
diff --git a/Engine/Plugins/SubtitlesAndClosedCaptions/Source/SubtitlesAndClosedCaptions/Public/SubtitlesAndClosedCaptionsModule.h b/Engine/Plugins/SubtitlesAndClosedCaptions/Source/SubtitlesAndClosedCaptions/Public/SubtitlesAndClosedCaptionsModule.h
new file mode 100644
--- /dev/null
+++ b/Engine/Plugins/SubtitlesAndClosedCaptions/Source/SubtitlesAndClosedCaptions/Public/SubtitlesAndClosedCaptionsModule.h
@@ -0,0 +1,19 @@
+// Copyright Epic Games, Inc. All Rights Reserved.
+
+#pragma once
+
+#include "Modules/ModuleInterface.h"
+#include "Subtitles/ISubtitlesAndClosedCaptionsModule.h"
+#include "Templates/SubclassOf.h"
+
+SUBTITLESANDCLOSEDCAPTIONS_API DECLARE_LOG_CATEGORY_EXTERN(LogSubtitlesAndClosedCaptions, Log, All);
+
+class UAssetUserData;
+
+class FSubtitlesAndClosedCaptionsModule : public ISubtitlesAndClosedCaptionsModule
+{
+public:
+	void StartupModule() override;
+	void ShutdownModule() override;
+	TSubclassOf<UAssetUserData> GetAssetUserDataClass() const override;
+};
diff --git a/Engine/Plugins/SubtitlesAndClosedCaptions/Source/SubtitlesAndClosedCaptions/Public/SubtitlesAudioSubsystem.h b/Engine/Plugins/SubtitlesAndClosedCaptions/Source/SubtitlesAndClosedCaptions/Public/SubtitlesAudioSubsystem.h
new file mode 100644
--- /dev/null
+++ b/Engine/Plugins/SubtitlesAndClosedCaptions/Source/SubtitlesAndClosedCaptions/Public/SubtitlesAudioSubsystem.h
@@ -0,0 +1,40 @@
+// Copyright Epic Games, Inc. All Rights Reserved.
+
+#pragma once 
+#include "Subsystems/AudioEngineSubsystem.h"
+#include "ActiveSoundUpdateInterface.h"
+#include "SubtitlesAudioSubsystem.generated.h"
+
+#define UE_API SUBTITLESANDCLOSEDCAPTIONS_API
+
+class UDialogueSoundWaveProxy;
+class USoundNode;
+
+/*
+*	SubtitlesAudioSubsystem - AudioEngineSubsystem for automatically queueing subtitles on audio classes.
+*
+*	Separated from SubtitlesSubsystem, which is a UWorldSubsystem, to avoid issues with lifecycle management
+*	as some parts need access to the world and others to the audio device.
+*/
+UCLASS(MinimalAPI, config = Game, defaultConfig)
+class USubtitlesAudioSubsystem : public UAudioEngineSubsystem
+	, public IActiveSoundUpdateInterface
+{
+	GENERATED_BODY()
+public:
+	USubtitlesAudioSubsystem() = default;
+
+	//~ Begin IActiveSoundUpdateInterface
+	virtual void NotifyActiveSoundCreated(FActiveSound& ActiveSound) override;
+	virtual void NotifyActiveSoundDeleting(const FActiveSound& ActiveSound) override;
+	//~ End IActiveSoundUpdateInterface
+
+private:
+	// Helper functions to cut down on duplicated code
+	static const UDialogueSoundWaveProxy* CheckSoundCueNodeForDialogueWaves(const USoundNode* Node, const FActiveSound& ActiveSound);
+	static void QueueSubtitleFromDialogueWaveProxy(const UDialogueSoundWaveProxy* DialogueWaveProxy, const FActiveSound& ActiveSound);
+
+	static void StopSubtitleFromDialogueWaveProxy(const UDialogueSoundWaveProxy* DialogueWaveProxy, const FActiveSound& ActiveSound);
+};
+
+#undef UE_API
diff --git a/Engine/Plugins/SubtitlesAndClosedCaptions/Source/SubtitlesAndClosedCaptions/Public/SubtitlesBlueprintFunctionLibrary.h b/Engine/Plugins/SubtitlesAndClosedCaptions/Source/SubtitlesAndClosedCaptions/Public/SubtitlesBlueprintFunctionLibrary.h
new file mode 100644
--- /dev/null
+++ b/Engine/Plugins/SubtitlesAndClosedCaptions/Source/SubtitlesAndClosedCaptions/Public/SubtitlesBlueprintFunctionLibrary.h
@@ -0,0 +1,56 @@
+// Copyright Epic Games, Inc. All Rights Reserved.
+
+#pragma once
+
+#include "Kismet/BlueprintFunctionLibrary.h"
+#include "Subtitles/SubtitlesAndClosedCaptionsTypes.h"
+#include "SubtitlesBlueprintFunctionLibrary.generated.h"
+
+#define UE_API SUBTITLESANDCLOSEDCAPTIONS_API
+
+class USubtitleWidget;
+
+UCLASS(MinimalAPI, meta = (ScriptName = "SubtitlesLibrary"))
+class USubtitlesBlueprintFunctionLibrary : public UBlueprintFunctionLibrary
+{
+	GENERATED_BODY()
+
+public:
+	// Functions taking USubtitleAssetUserData as a parameter use TSoftObjectPtr to ensure that you can drag an asset onto BP variables feeding into them.
+	// A const pointer allows drag/drop onto the function parameter itself, but TSoftObjectPtr is required for anything with the DefaultToInstanced specifier
+	// (which in this case is the parent class UAssetUserData).
+
+	// Queues all of the subtitles in a Subtitle asset, accounting for any delayed start offsets set for each. Subtitles queued through this function are internally-timed.
+	UFUNCTION(BlueprintCallable, Category = Subtitles)
+	static UE_API void QueueSubtitlesFromAsset(TSoftObjectPtr<USubtitleAssetUserData> SubtitleAsset);
+
+	// Stops all subtitles in a subtitle asset, regardless of whether they're already queued or if they have a delayed start.
+	UFUNCTION(BlueprintCallable, Category = Subtitles)
+	static UE_API void StopSubtitlesInAsset(TSoftObjectPtr<USubtitleAssetUserData> SubtitleAsset);
+
+	// Queues an individual subtitle from a struct, or a single entry in a Subtitle asset's array.
+	UE_DEPRECATED(5.8, "Use QueueSingleSubtitle. It removes extraneous and confusing parameters.")
+	UFUNCTION(BlueprintCallable, Category = Subtitles)
+	static UE_API void QueueSubtitle(const FSubtitleAssetData& Subtitle, const ESubtitleTiming Timing = ESubtitleTiming::InternallyTimed);
+
+	// Queues an individual subtitle.	
+	// Parameters form an FSubtitleAssetData (see SubtitlesAndClosedCaptionsDelegates.h), but the Duration Type is removed as it is irrelevant for queueing this way.
+	UFUNCTION(BlueprintCallable, Category = Subtitles)
+	static UE_API void QueueSingleSubtitle(
+		const FText Text, const float Duration, const float StartOffset, const float Priority, const ESubtitleType SubtitleType, const ESubtitleTiming Timing = ESubtitleTiming::InternallyTimed);
+
+
+	UFUNCTION(BlueprintCallable, Category = Subtitles)
+	static UE_API bool IsSubtitleActive(const FSubtitleAssetData& Subtitle);
+
+	UFUNCTION(BlueprintCallable, Category = Subtitles)
+	static UE_API void StopSubtitle(const FSubtitleAssetData& Subtitle);
+
+	UFUNCTION(BlueprintCallable, Category = Subtitles)
+	static UE_API void StopAllSubtitles();
+
+	UFUNCTION(BlueprintCallable, Category = Subtitles)
+	static UE_API void ReplaceSubtitleWidget(const TSoftClassPtr<USubtitleWidget>& NewWidgetAsset);
+};
+
+#undef UE_API
diff --git a/Engine/Plugins/SubtitlesAndClosedCaptions/Source/SubtitlesAndClosedCaptions/Public/SubtitlesSettings.h b/Engine/Plugins/SubtitlesAndClosedCaptions/Source/SubtitlesAndClosedCaptions/Public/SubtitlesSettings.h
new file mode 100644
--- /dev/null
+++ b/Engine/Plugins/SubtitlesAndClosedCaptions/Source/SubtitlesAndClosedCaptions/Public/SubtitlesSettings.h
@@ -0,0 +1,38 @@
+// Copyright Epic Games, Inc. All Rights Reserved.
+
+#pragma once
+
+#include "Engine/DeveloperSettings.h"
+#include "Math/Color.h"
+#include "SubtitleWidget.h"
+#include "SubtitlesSettings.generated.h"
+
+
+UCLASS(config = Game, defaultconfig, meta = (DisplayName = "Subtitles And Closed Captions"))
+class USubtitlesSettings : public UDeveloperSettings
+{
+	GENERATED_BODY()
+
+public:
+	USubtitlesSettings();
+
+	// Temporarily retaining the old GetWidget to not break builds.
+	// Plugin is still experimental, so we don't need to actually deprecate it just yet.
+	const TSubclassOf<USubtitleWidget>& GetWidget() const { return WidgetToUseAsSubclass; };
+	const TSubclassOf<USubtitleWidget>& GetWidgetDefault() const { return DefaultWidgetToUseAsSubclass;	};
+private:
+	TSubclassOf<USubtitleWidget> WidgetToUseAsSubclass;
+	TSubclassOf<USubtitleWidget> DefaultWidgetToUseAsSubclass;
+
+public:
+	inline const TSoftClassPtr<USubtitleWidget>& GetWidgetSoftClassPtr() const { return SubtitleWidgetToUse; }
+	inline const TSoftClassPtr<USubtitleWidget>& GetWidgetDefaultSoftClassPtr() const { return SubtitleWidgetToUseDefault; }
+
+protected:
+
+	UPROPERTY(config, EditDefaultsOnly, Category = Subtitles, AdvancedDisplay)
+	TSoftClassPtr<USubtitleWidget> SubtitleWidgetToUse;
+
+	UPROPERTY(config)
+	TSoftClassPtr<USubtitleWidget> SubtitleWidgetToUseDefault; // fallback for SubtitleWidgetToUse (not set by user)
+};
diff --git a/Engine/Plugins/SubtitlesAndClosedCaptions/Source/SubtitlesAndClosedCaptions/Public/SubtitlesSubsystem.h b/Engine/Plugins/SubtitlesAndClosedCaptions/Source/SubtitlesAndClosedCaptions/Public/SubtitlesSubsystem.h
new file mode 100644
--- /dev/null
+++ b/Engine/Plugins/SubtitlesAndClosedCaptions/Source/SubtitlesAndClosedCaptions/Public/SubtitlesSubsystem.h
@@ -0,0 +1,116 @@
+// Copyright Epic Games, Inc. All Rights Reserved.
+
+#pragma once 
+
+#include "ActiveSubtitleAssetData.h"
+#include "Fonts/SlateFontInfo.h"
+#include "Math/Color.h"
+#include "Subsystems/WorldSubsystem.h"
+#include "Subtitles/SubtitlesAndClosedCaptionsDelegates.h"
+
+#include "SubtitleWidget.h"
+
+#include "SubtitlesSubsystem.generated.h"
+
+#define UE_API SUBTITLESANDCLOSEDCAPTIONS_API
+
+class FCanvas;
+class UAssetUserData;
+class UFont;
+struct FQueueSubtitleParameters;
+
+#if WITH_DEV_AUTOMATION_TESTS
+namespace SubtitlesAndClosedCaptions::Test
+{
+#if 0 // Temporarily disabling these tests as they have a dangling reference that trips up the static analysis on certain build configurations.
+	struct FMovieSceneSubtitlesTest;
+#endif
+	struct FSubtitlesTest;
+}
+#endif
+
+/*
+* #SUBTITLES_PRD -	Requirement:	Ability to allow designers to “script” subtitle location for sequences and scenes to avoid subtitles overlapping important scenes or characters
+*					Use a UEngineSubsystem for blueprints
+*
+*	Game configuration for font customization per game
+*/
+UCLASS(MinimalAPI, config = Game, defaultConfig)
+class USubtitlesSubsystem : public UWorldSubsystem
+{
+	GENERATED_BODY()
+public:
+	USubtitlesSubsystem() = default;
+
+	// FSubtitlesAndClosedCaptionsDelegates
+
+	// Adds a subtitle to the queue: Params contains the subtitle asset and an optional duration. The highest-priority subtitle in the queue will be displayed.
+	// If Timing is ExternallyTimed, the queued subtitle will remain in the queue until manually removed.
+	// If the subtitle asset has a non-zero StartOffset, it will sit in a delayed-start queue instead of being queued for display.
+	UE_API virtual void QueueSubtitle(FQueueSubtitleParameters Params, const ESubtitleTiming Timing = ESubtitleTiming::InternallyTimed);
+
+	// Returns true if the given subtitle asset is being displayed.
+	UE_API virtual bool IsSubtitleActive(FSubtitleAssetData Data) const;
+
+	// Stops the given subtitle asset being displayed.  This includes subtitles not yet being displayed due to their StartOffset
+	UE_API virtual void StopSubtitle(FSubtitleAssetData Data);
+
+	// Stops all queued subtitles from being displayed.  This includes subtitles not yet being displayed due to their StartOffset.
+	UE_API virtual void StopAllSubtitles();
+
+	// Replaces the Subtitle Widget with a new asset subclassed from the base USubtitleWidget. Used by the Replace Subtitle Widget BP function.
+	UE_API virtual void ReplaceWidget(const TSoftClassPtr<USubtitleWidget>& NewWidgetAsset);
+
+protected:
+
+	UE_API virtual void Initialize(FSubsystemCollectionBase& Collection) override;
+	UE_API virtual void BindDelegates();
+
+	UPROPERTY()
+	TMap<ESubtitleType, FSlateFontInfo> SubtitleFontInfo;
+
+	// Sorted by priority, desc
+	UPROPERTY()
+	TArray<FActiveSubtitleAssetData> ActiveSubtitles;
+
+	// Unsorted; subtitles with a delayed start offset still need to be tracked before entering the queue proper.
+	UPROPERTY()
+	TArray<FActiveSubtitleAssetData> DelayedSubtitles;
+
+protected:
+	UE_API virtual void AddActiveSubtitle(const FSubtitleAssetData& Subtitle, float Duration, const float StartOffset, const ESubtitleTiming Timing);
+
+	// These need to be UFUNCTIONs for the Timer Delegate bindings.
+	UFUNCTION()
+	UE_API virtual void MakeDelayedSubtitleActive(const FSubtitleAssetData& Subtitle, const ESubtitleTiming Timing);
+
+	UFUNCTION()
+	UE_API virtual void RemoveActiveSubtitle(const FSubtitleAssetData& Subtitle);
+
+private:
+	UE_API void AddAndDisplaySubtitle(FActiveSubtitleAssetData& NewActiveSubtitle);
+
+public:
+	// GetActiveSubtitles returns the sorted array of queued subtitles. Higher priorities are earlier in the array.
+	// For equally-prioritized subtitles, the oldest one is first in this array.
+	const TArray<FActiveSubtitleAssetData>& GetActiveSubtitles() const
+	{
+		return ActiveSubtitles;
+	}
+
+	inline TWeakObjectPtr<USubtitleWidget> GetActiveSubtitleWidget()
+	{
+		return SubtitleWidget;
+	}
+
+private:
+	UPROPERTY(Transient)
+	TObjectPtr<USubtitleWidget> SubtitleWidget = nullptr;
+
+	bool bInitializedWidget = false;
+	UE_API bool TryCreateUMGWidgetFromAsset(const TSoftClassPtr<USubtitleWidget>& WidgetToUse);
+	UE_API bool TryCreateUMGWidget();
+	UE_API void UpdateWidgetData();
+};
+
+#undef UE_API
diff --git a/Engine/Plugins/SubtitlesAndClosedCaptions/Source/SubtitlesAndClosedCaptions/SubtitlesAndClosedCaptions.Build.cs b/Engine/Plugins/SubtitlesAndClosedCaptions/Source/SubtitlesAndClosedCaptions/SubtitlesAndClosedCaptions.Build.cs
new file mode 100644
--- /dev/null
+++ b/Engine/Plugins/SubtitlesAndClosedCaptions/Source/SubtitlesAndClosedCaptions/SubtitlesAndClosedCaptions.Build.cs
@@ -0,0 +1,35 @@
+// Copyright Epic Games, Inc. All Rights Reserved.
+
+namespace UnrealBuildTool.Rules
+{
+	public class SubtitlesAndClosedCaptions : ModuleRules
+	{
+		public SubtitlesAndClosedCaptions(ReadOnlyTargetRules Target) : base(Target)
+		{
+			PCHUsage = PCHUsageMode.UseExplicitOrSharedPCHs;
+
+			PrivateDependencyModuleNames.AddRange
+			(
+				new string[]
+				{
+					"Slate",
+					"SlateCore",
+				}
+			);
+
+			PublicDependencyModuleNames.AddRange(
+				new string[]
+				{
+					"Core",
+					"CoreUObject",
+					"DeveloperSettings",
+					"Engine",
+					"InputCore",
+					"MovieSceneTracks",
+					"MovieScene",
+					"UMG"
+				}
+			);
+		}
+	}
+}
diff --git a/Engine/Plugins/SubtitlesAndClosedCaptions/Source/SubtitlesAndClosedCaptionsEditor/Private/SubtitleAssetActionUtility.cpp b/Engine/Plugins/SubtitlesAndClosedCaptions/Source/SubtitlesAndClosedCaptionsEditor/Private/SubtitleAssetActionUtility.cpp
new file mode 100644
--- /dev/null
+++ b/Engine/Plugins/SubtitlesAndClosedCaptions/Source/SubtitlesAndClosedCaptionsEditor/Private/SubtitleAssetActionUtility.cpp
@@ -0,0 +1,56 @@
+// Copyright Epic Games, Inc. All Rights Reserved.
+
+#include "SubtitleAssetActionUtility.h"
+#include "AssetRegistry/AssetRegistryModule.h"
+#include "AssetToolsModule.h"
+#include "BasicOverlays.h"
+#include "EditorUtilityLibrary.h"
+#include "IAssetTools.h"
+#include "Subtitles/SubtitlesAndClosedCaptionsDelegates.h"
+
+USubtitleAssetActionUtility::USubtitleAssetActionUtility()
+{
+	SupportedClasses.Add(UBasicOverlays::StaticClass());
+}
+
+// .SRT (SubRip Text) files are imported by the engine as UBasicOverlays, which contain an Array
+// of captions with timestamps and positions. This utility creates a new asset that converts BasicOverlays to Subtitle assets.
+void USubtitleAssetActionUtility::ConvertBasicOverlaysToSubtitles()
+{
+	TArray<UObject*> SelectedAssets = UEditorUtilityLibrary::GetSelectedAssets();
+
+	for (UObject* Asset : SelectedAssets)
+	{
+		UBasicOverlays* OverlayAsset = Cast<UBasicOverlays>(Asset);
+		if (IsValid(OverlayAsset))
+		{
+			const FString NameSuffix = TEXT("Subtitle");
+			const FString AssetPath = FPaths::ProjectContentDir();
+			FString AssetName = OverlayAsset->GetName() + NameSuffix;
+			FString PackagePath = OverlayAsset->GetPackage()->GetPathName() + NameSuffix;
+
+			UPackage* Package = CreatePackage(*PackagePath);
+			USubtitleAssetUserData* NewSubtitleAsset = NewObject<USubtitleAssetUserData>(
+				Package, *AssetName, EObjectFlags::RF_Public | EObjectFlags::RF_Standalone
+			);
+
+			// Iterate through each timecode in the overlay and create a new subtitle from it.
+			for (const FOverlayItem& Overlay : OverlayAsset->GetAllOverlays())
+			{
+				FSubtitleAssetData Subtitle;
+				Subtitle.Text = FText::FromString(Overlay.Text);
+				
+				// Timing
+				Subtitle.StartOffset = Overlay.StartTime.GetTotalSeconds();
+				Subtitle.Duration = Overlay.EndTime.GetTotalSeconds() - Overlay.StartTime.GetTotalSeconds();
+				Subtitle.SubtitleDurationType = ESubtitleDurationType::UseDurationProperty;
+				
+				// Default priority, subtitle type, etc are fine.
+				NewSubtitleAsset->Subtitles.Add(Subtitle);
+			}
+
+			FAssetRegistryModule::AssetCreated(NewSubtitleAsset);
+			NewSubtitleAsset->MarkPackageDirty();
+		}
+	}
+}
\ No newline at end of file
diff --git a/Engine/Plugins/SubtitlesAndClosedCaptions/Source/SubtitlesAndClosedCaptionsEditor/Private/SubtitleAssetDefinitions.h b/Engine/Plugins/SubtitlesAndClosedCaptions/Source/SubtitlesAndClosedCaptionsEditor/Private/SubtitleAssetDefinitions.h
new file mode 100644
--- /dev/null
+++ b/Engine/Plugins/SubtitlesAndClosedCaptions/Source/SubtitlesAndClosedCaptionsEditor/Private/SubtitleAssetDefinitions.h
@@ -0,0 +1,28 @@
+// Copyright Epic Games, Inc. All Rights Reserved.
+
+#pragma once
+
+#include "AssetDefinitionDefault.h"
+#include "Subtitles/SubtitlesAndClosedCaptionsDelegates.h"
+#include "SubtitlesAndClosedCaptionsEditorModule.h"
+
+#include "SubtitleAssetDefinitions.generated.h"
+
+UCLASS()
+class USubtitleAssetDefinitions : public UAssetDefinitionDefault
+{
+	GENERATED_BODY()
+
+public:
+	// UAssetDefinition Begin
+	virtual FText GetAssetDisplayName() const override { return NSLOCTEXT("AssetTypeActions", "AssetTypeActions_Subtitle", "Subtitle"); }
+	virtual FLinearColor GetAssetColor() const override { return FLinearColor(FColor(97, 85, 212)); }
+	virtual TSoftClassPtr<UObject> GetAssetClass() const override { return USubtitleAssetUserData::StaticClass(); }
+	virtual TConstArrayView<FAssetCategoryPath> GetAssetCategories() const override
+	{
+
+		static const FAssetCategoryPath Categories[1] = {FAssetCategoryPath(ISubtitlesAndClosedCaptionsEditorModule::GetAssetTypeCategory())};
+		return Categories;
+	}
+	// UAssetDefinition End
+};
diff --git a/Engine/Plugins/SubtitlesAndClosedCaptions/Source/SubtitlesAndClosedCaptionsEditor/Private/SubtitleFactory.cpp b/Engine/Plugins/SubtitlesAndClosedCaptions/Source/SubtitlesAndClosedCaptionsEditor/Private/SubtitleFactory.cpp
new file mode 100644
--- /dev/null
+++ b/Engine/Plugins/SubtitlesAndClosedCaptions/Source/SubtitlesAndClosedCaptionsEditor/Private/SubtitleFactory.cpp
@@ -0,0 +1,20 @@
+// Copyright Epic Games, Inc. All Rights Reserved.
+#include "SubtitleFactory.h"
+
+#include "Subtitles/SubtitlesAndClosedCaptionsDelegates.h"
+
+#include UE_INLINE_GENERATED_CPP_BY_NAME(SubtitleFactory)
+
+USubtitleFactory::USubtitleFactory(const FObjectInitializer& ObjectInitializer)
+	: Super(ObjectInitializer)
+{
+	bCreateNew = true;
+	bEditorImport = false;
+	bEditAfterNew = true;
+	SupportedClass = USubtitleAssetUserData::StaticClass();
+}
+
+UObject* USubtitleFactory::FactoryCreateNew(UClass* InClass, UObject* InParent, FName InName, EObjectFlags InFlags, UObject* InContext, FFeedbackContext* InFeedbackContext)
+{
+	return NewObject<USubtitleAssetUserData>(InParent, InClass, InName, InFlags);
+}
diff --git a/Engine/Plugins/SubtitlesAndClosedCaptions/Source/SubtitlesAndClosedCaptionsEditor/Private/SubtitleFactory.h b/Engine/Plugins/SubtitlesAndClosedCaptions/Source/SubtitlesAndClosedCaptionsEditor/Private/SubtitleFactory.h
new file mode 100644
--- /dev/null
+++ b/Engine/Plugins/SubtitlesAndClosedCaptions/Source/SubtitlesAndClosedCaptionsEditor/Private/SubtitleFactory.h
@@ -0,0 +1,18 @@
+// Copyright Epic Games, Inc. All Rights Reserved.
+
+#pragma once
+
+#include "Factories/Factory.h"
+#include "UObject/ObjectMacros.h"
+
+#include "SubtitleFactory.generated.h"
+
+UCLASS(hidecategories = Object, MinimalAPI)
+class USubtitleFactory : public UFactory
+{
+	GENERATED_UCLASS_BODY()
+
+	// UFactory Begin
+	virtual UObject* FactoryCreateNew(UClass* InClass, UObject* InParent, FName InName, EObjectFlags InFlags, UObject* InContext, FFeedbackContext* InFeedbackContext) override;
+	// UFactory End
+};
diff --git a/Engine/Plugins/SubtitlesAndClosedCaptions/Source/SubtitlesAndClosedCaptionsEditor/Private/SubtitleMigrationUtility.cpp b/Engine/Plugins/SubtitlesAndClosedCaptions/Source/SubtitlesAndClosedCaptionsEditor/Private/SubtitleMigrationUtility.cpp
new file mode 100644
--- /dev/null
+++ b/Engine/Plugins/SubtitlesAndClosedCaptions/Source/SubtitlesAndClosedCaptionsEditor/Private/SubtitleMigrationUtility.cpp
@@ -0,0 +1,159 @@
+// Copyright Epic Games, Inc. All Rights Reserved.
+
+#include "SubtitleMigrationUtility.h"
+
+#include "EditorUtilityLibrary.h"
+#include "Engine/Engine.h"
+#include "Logging/StructuredLog.h"
+#include "Sound/SoundWave.h"
+#include "SubtitlesAndClosedCaptionsModule.h"
+#include "Subtitles/SubtitlesAndClosedCaptionsDelegates.h"
+
+const FString USubtitleMigrationUtility::UndoContext = TEXT("SubtitleMigration");
+
+USubtitleMigrationUtility::USubtitleMigrationUtility()
+{
+	SupportedClasses.Add(USoundBase::StaticClass());
+}
+
+void USubtitleMigrationUtility::AddBlankSubtitle()
+{
+#if WITH_EDITOR
+	if (GEngine)
+	{
+		// Setup transaction for undoing, and describe said action.
+		const static FText Description = FText::FromString(TEXT("Add Blank Subtitle"));
+		GEngine->BeginTransaction(*UndoContext, Description, nullptr);
+
+		const TArray<UObject*> SelectedAssets = UEditorUtilityLibrary::GetSelectedAssetsOfClass(USoundBase::StaticClass());
+		for (UObject* SelectedObject : SelectedAssets)
+		{
+			if (!IsValid(SelectedObject))
+			{
+				continue;
+			}
+
+			// GetSelectedAssetsOfClass and SupportedClasses should already only provide USoundBase
+			USoundBase* SelectedSound = CastChecked<USoundBase>(SelectedObject);
+
+			// Add an asset userdata if one doesn't already exist.
+			if (!SelectedSound->HasAssetUserDataOfClass(USubtitleAssetUserData::StaticClass()))
+			{
+				SelectedSound->AddAssetUserDataOfClass(USubtitleAssetUserData::StaticClass());
+			}
+
+			UAssetUserData* UserData = SelectedSound->GetAssetUserDataOfClass(USubtitleAssetUserData::StaticClass());
+			USubtitleAssetUserData* SubtitleData = CastChecked<USubtitleAssetUserData>(UserData);
+
+			// Mark the userdata as modified for the undo action to work correctly.
+			SubtitleData->Modify();
+
+			SubtitleData->Subtitles.Add(FSubtitleAssetData());
+		}
+
+		// End transaction for undoing.
+		GEngine->EndTransaction();
+	}
+#endif
+}
+
+void USubtitleMigrationUtility::RemoveLegacySubtitles()
+{
+#if WITH_EDITOR
+	if (GEngine)
+	{
+		// Setup transaction for undoing, and describe said action.
+		const static FText Description = FText::FromString(TEXT("Remove legacy subtitles from SoundWaves."));
+		GEngine->BeginTransaction(*UndoContext, Description, nullptr);
+
+		const TArray<UObject*> SelectedAssets = UEditorUtilityLibrary::GetSelectedAssetsOfClass(USoundWave::StaticClass());
+		for (UObject* SelectedObject : SelectedAssets)
+		{
+			if (!IsValid(SelectedObject))
+			{
+				continue;
+			}
+
+			// This tool supports other non-wave Soundbases for other actions, but only USoundWaves have subtitles to migrate in this fashion.
+			USoundWave* SelectedSound = Cast<USoundWave>(SelectedObject);
+
+			// Only empty out the legacy subtitles if there's already new subtitles in place.
+			if (IsValid(SelectedSound))
+			{
+				if (SelectedSound->HasAssetUserDataOfClass(USubtitleAssetUserData::StaticClass()))
+				{
+					SelectedSound->Modify();
+					SelectedSound->Subtitles.Empty();
+				}
+				else
+				{
+					UE_LOGFMT(LogSubtitlesAndClosedCaptions, Warning, "Can't remove legacy subtitles: migrate them to the new system first with Convert Legacy Subtitles.");
+				}
+			}
+			else
+			{
+				UE_LOGFMT(LogSubtitlesAndClosedCaptions, Warning, "Remove Legacy Subtitles can only be used on SoundWave assets.");
+			}
+		}
+
+		GEngine->EndTransaction();
+	}
+#endif
+}
+
+void USubtitleMigrationUtility::ConvertLegacySubtitles()
+{
+#if WITH_EDITOR
+	if (GEngine)
+	{
+		// Setup transaction for undoing, and describe said action.
+		const static FText Description = FText::FromString(TEXT("Make Asset UserData from Legacy Subtitles"));
+		GEngine->BeginTransaction(*UndoContext, Description, nullptr);
+
+		const TArray<UObject*> SelectedAssets = UEditorUtilityLibrary::GetSelectedAssetsOfClass(USoundWave::StaticClass());
+		for (UObject* SelectedObject : SelectedAssets)
+		{
+			if (!IsValid(SelectedObject))
+			{
+				continue;
+			}
+
+			// This tool supports other non-wave Soundbases for other actions, but only USoundWaves have subtitles to migrate in this fashion.
+			USoundWave* SelectedSound = Cast<USoundWave>(SelectedObject);
+
+			if (IsValid(SelectedSound))
+			{
+				// Add an asset userdata if one doesn't already exist.
+				if (!SelectedSound->HasAssetUserDataOfClass(USubtitleAssetUserData::StaticClass()))
+				{
+					SelectedSound->AddAssetUserDataOfClass(USubtitleAssetUserData::StaticClass());
+				}
+
+				UAssetUserData* UserData = SelectedSound->GetAssetUserDataOfClass(USubtitleAssetUserData::StaticClass());
+				USubtitleAssetUserData* SubtitleData = CastChecked<USubtitleAssetUserData>(UserData);
+
+				// Loop through all the legacy subtitles and add a new subtitle to the UserData for each.
+				for (const FSubtitleCue& Cue : SelectedSound->Subtitles)
+				{
+					FSubtitleAssetData Data;
+					Data.SubtitleType = ESubtitleType::Subtitle;
+					Data.SubtitleDurationType = ESubtitleDurationType::UseSoundDuration;
+
+					Data.Priority = SelectedSound->SubtitlePriority;
+					Data.Text = Cue.Text;
+					Data.StartOffset = Cue.Time;
+					Data.Comment = SelectedSound->Comment; // While it looks like this is the sound's comment, the property is in the Subtitles category.
+					SubtitleData->Modify();
+					SubtitleData->Subtitles.Add(Data);
+				}
+			}
+			else
+			{
+				UE_LOGFMT(LogSubtitlesAndClosedCaptions, Warning, "Convert Legacy Subtitles can only be used on SoundWave assets.");
+			}
+		}
+
+		GEngine->EndTransaction();
+	}
+#endif
+}
\ No newline at end of file
diff --git a/Engine/Plugins/SubtitlesAndClosedCaptions/Source/SubtitlesAndClosedCaptionsEditor/Private/SubtitleSequencerSection.cpp b/Engine/Plugins/SubtitlesAndClosedCaptions/Source/SubtitlesAndClosedCaptionsEditor/Private/SubtitleSequencerSection.cpp
new file mode 100644
--- /dev/null
+++ b/Engine/Plugins/SubtitlesAndClosedCaptions/Source/SubtitlesAndClosedCaptionsEditor/Private/SubtitleSequencerSection.cpp
@@ -0,0 +1,36 @@
+// Copyright Epic Games, Inc. All Rights Reserved.
+
+#include "SubtitleSequencerSection.h"
+
+#include "MovieSceneSection.h"
+#include "MovieSceneSubtitleSection.h"
+#include "SequencerSectionPainter.h"
+#include "Subtitles/SubtitlesAndClosedCaptionsDelegates.h"
+
+FText FSubtitleSequencerSection::GetSectionTitle() const
+{
+	if (TStrongObjectPtr<const UMovieSceneSection> SectionPin = WeakSection.Pin())
+	{
+		const UMovieSceneSubtitleSection& SubtitlesSection = *CastChecked<const UMovieSceneSubtitleSection>(SectionPin.Get());
+		if (const USubtitleAssetUserData* SubtitleAsset = SubtitlesSection.GetSubtitle())
+		{
+			const TArray<FSubtitleAssetData>& Subtitles = SubtitleAsset->Subtitles;
+			if (!Subtitles.IsEmpty())
+			{
+				return Subtitles[0].Text;
+			}
+			else
+			{
+				return NSLOCTEXT("FSubtitlesSequencerSection", "EmptySubtitleAsset", "Empty Subtitle");
+			}
+		}
+	}
+
+	return NSLOCTEXT("FSubtitlesSequencerSection", "NoSubtitleName", "No Subtitle");
+}
+
+int32 FSubtitleSequencerSection::OnPaintSection(FSequencerSectionPainter& Painter) const
+{
+	int32 LayerId = Painter.PaintSectionBackground();
+	return LayerId;
+}
diff --git a/Engine/Plugins/SubtitlesAndClosedCaptions/Source/SubtitlesAndClosedCaptionsEditor/Private/SubtitleSequencerSection.h b/Engine/Plugins/SubtitlesAndClosedCaptions/Source/SubtitlesAndClosedCaptionsEditor/Private/SubtitleSequencerSection.h
new file mode 100644
--- /dev/null
+++ b/Engine/Plugins/SubtitlesAndClosedCaptions/Source/SubtitlesAndClosedCaptionsEditor/Private/SubtitleSequencerSection.h
@@ -0,0 +1,21 @@
+// Copyright Epic Games, Inc. All Rights Reserved.
+
+#pragma once
+
+#include "ISequencerSection.h"
+
+class FSequencerSectionPainter;
+
+class FSubtitleSequencerSection
+	: public FSequencerSection
+	, public TSharedFromThis<FSubtitleSequencerSection>
+{
+public:
+	FSubtitleSequencerSection(UMovieSceneSection& InSection)
+		: FSequencerSection(InSection)
+	{}
+
+	// ISequencerSection interface
+	virtual FText GetSectionTitle() const override;
+	virtual int32 OnPaintSection(FSequencerSectionPainter& Painter) const override;
+};
diff --git a/Engine/Plugins/SubtitlesAndClosedCaptions/Source/SubtitlesAndClosedCaptionsEditor/Private/SubtitlesAndClosedCaptionsEditorModule.cpp b/Engine/Plugins/SubtitlesAndClosedCaptions/Source/SubtitlesAndClosedCaptionsEditor/Private/SubtitlesAndClosedCaptionsEditorModule.cpp
new file mode 100644
--- /dev/null
+++ b/Engine/Plugins/SubtitlesAndClosedCaptions/Source/SubtitlesAndClosedCaptionsEditor/Private/SubtitlesAndClosedCaptionsEditorModule.cpp
@@ -0,0 +1,46 @@
+// Copyright Epic Games, Inc. All Rights Reserved.
+
+#include "SubtitlesAndClosedCaptionsEditorModule.h"
+
+#include "AssetToolsModule.h"
+#include "AssetTypeCategories.h"
+#include "ISequencerModule.h"
+#include "Modules/ModuleManager.h"
+#include "SubtitlesTrackEditor.h"
+
+#define LOCTEXT_NAMESPACE "SubtitlesAndClosedCaptionsEditor"
+
+FText ISubtitlesAndClosedCaptionsEditorModule::AssetTypeCategory = LOCTEXT("SubtitlesAssetTypeCategory", "Subtitles");
+
+class FSubtitlesAndClosedCaptionsEditorModule : public ISubtitlesAndClosedCaptionsEditorModule
+{
+public:
+	virtual void StartupModule() override
+	{
+		FModuleManager::Get().LoadModuleChecked("SubtitlesAndClosedCaptions");
+
+		// Register assets
+		IAssetTools& AssetTools = FModuleManager::LoadModuleChecked<FAssetToolsModule>("AssetTools").Get();
+		AssetTools.RegisterAdvancedAssetCategory(FName(TEXT("Subtitles")), AssetTypeCategory);
+
+		//
+		ISequencerModule& SequencerModule = FModuleManager::LoadModuleChecked<ISequencerModule>("Sequencer");
+		CreateTrackEditorHandle = SequencerModule.RegisterTrackEditor(FOnCreateTrackEditor::CreateStatic(&FSubtitlesTrackEditor::CreateTrackEditor));
+
+	}
+
+	virtual void ShutdownModule() override
+	{
+		ISequencerModule* SequencerModule = FModuleManager::GetModulePtr<ISequencerModule>("Sequencer");
+		if (SequencerModule != nullptr)
+		{
+			SequencerModule->UnRegisterTrackEditor(CreateTrackEditorHandle);
+		}
+	}
+
+	FDelegateHandle CreateTrackEditorHandle;
+};
+
+IMPLEMENT_MODULE(FSubtitlesAndClosedCaptionsEditorModule, SubtitlesAndClosedCaptionsEditor);
+
+#undef LOCTEXT_NAMESPACE
diff --git a/Engine/Plugins/SubtitlesAndClosedCaptions/Source/SubtitlesAndClosedCaptionsEditor/Private/SubtitlesTrackEditor.cpp b/Engine/Plugins/SubtitlesAndClosedCaptions/Source/SubtitlesAndClosedCaptionsEditor/Private/SubtitlesTrackEditor.cpp
new file mode 100644
--- /dev/null
+++ b/Engine/Plugins/SubtitlesAndClosedCaptions/Source/SubtitlesAndClosedCaptionsEditor/Private/SubtitlesTrackEditor.cpp
@@ -0,0 +1,422 @@
+// Copyright Epic Games, Inc. All Rights Reserved.
+
+#include "SubtitlesTrackEditor.h"
+
+#include "AssetRegistry/AssetRegistryModule.h"
+#include "ContentBrowserDelegates.h"
+#include "ContentBrowserModule.h"
+#include "DragAndDrop/AssetDragDropOp.h"
+#include "IContentBrowserSingleton.h"
+#include "Input/Reply.h"
+#include "LevelSequence.h"
+#include "MovieSceneSubtitlesTrack.h"
+#include "MovieSceneToolHelpers.h"
+#include "MVVM/Views/ViewUtilities.h"
+#include "SequencerSettings.h"
+#include "Subtitles/SubtitlesAndClosedCaptionsDelegates.h"
+#include "SubtitleSequencerSection.h"
+
+
+#define LOCTEXT_NAMESPACE "FSubtitlesTrackEditor"
+
+FSubtitlesTrackEditor::FSubtitlesTrackEditor(TSharedRef<ISequencer> InSequencer)
+	: FMovieSceneTrackEditor(InSequencer)
+{
+}
+
+TSharedRef<ISequencerTrackEditor> FSubtitlesTrackEditor::CreateTrackEditor(TSharedRef<ISequencer> OwningSequencer)
+{
+	return MakeShareable(new FSubtitlesTrackEditor(OwningSequencer));
+}
+
+FText FSubtitlesTrackEditor::GetDisplayName() const
+{
+	return LOCTEXT("SubtitilesTrackEditor_DisplayName", "Subtitles");
+}
+
+void FSubtitlesTrackEditor::BuildAddTrackMenu(FMenuBuilder& MenuBuilder)
+{
+	MenuBuilder.AddMenuEntry(
+		LOCTEXT("AddTrack", "Subtitles Track"),
+		LOCTEXT("AddTooltip", "Adds a new subtitles track that can display subtitles and closed captions."),
+		FSlateIcon(FAppStyle::GetAppStyleSetName(), "Sequencer.Tracks.Audio"), // #SUBTITLES_PRD_TODO - Add a subtitles track icon
+		FUIAction(
+			FExecuteAction::CreateRaw(this, &FSubtitlesTrackEditor::HandleAddMenuEntryExecute)
+		)
+	);
+}
+
+void FSubtitlesTrackEditor::HandleAddMenuEntryExecute()
+{
+	UMovieScene* FocusedMovieScene = GetFocusedMovieScene();
+
+	if (FocusedMovieScene == nullptr)
+	{
+		return;
+	}
+
+	if (FocusedMovieScene->IsReadOnly())
+	{
+		return;
+	}
+
+	const FScopedTransaction Transaction(NSLOCTEXT("Sequencer", "AddSubtitlesTrack_Transaction", "Add Subtitles Track"));
+	FocusedMovieScene->Modify();
+
+	UMovieSceneSubtitlesTrack* NewTrack = FocusedMovieScene->AddTrack<UMovieSceneSubtitlesTrack>();
+	if (IsValid(NewTrack))
+	{
+		NewTrack->SetDisplayName(LOCTEXT("SubtitlesTrackName", "Subtitles"));
+
+		if (GetSequencer().IsValid())
+		{
+			GetSequencer()->OnAddTrack(NewTrack, FGuid());
+		}
+	}
+}
+
+bool FSubtitlesTrackEditor::IsDragDropValid(const FDragDropEvent& DragDropEvent, const FSequencerDragDropParams& DragDropParams, DragDropValidInfo& OutDragDropInfo)
+{
+	if (!DragDropParams.Track.IsValid() || !DragDropParams.Track.Get()->IsA(UMovieSceneSubtitlesTrack::StaticClass()))
+	{
+		return false;
+	}
+	OutDragDropInfo.Operation = DragDropEvent.GetOperation();
+	if (!OutDragDropInfo.Operation.IsValid() || !OutDragDropInfo.Operation->IsOfType<FAssetDragDropOp>())
+	{
+		return false;
+	}
+	OutDragDropInfo.SequencerPtr = GetSequencer();
+	if (!OutDragDropInfo.SequencerPtr)
+	{
+		return false;
+	}
+	OutDragDropInfo.FocusedSequence = OutDragDropInfo.SequencerPtr->GetFocusedMovieSceneSequence();
+	if (!OutDragDropInfo.FocusedSequence)
+	{
+		return false;
+	}
+	OutDragDropInfo.DragDropOp = StaticCastSharedPtr<FAssetDragDropOp>(OutDragDropInfo.Operation);
+
+	return true;
+}
+
+TSharedPtr<SWidget> FSubtitlesTrackEditor::BuildOutlinerEditWidget(const FGuid& ObjectBinding, UMovieSceneTrack* Track, const FBuildEditWidgetParams& Params)
+{
+	TDelegate OnAssetSelected = FOnAssetSelected::CreateRaw(this, &FSubtitlesTrackEditor::OnAssetSelected, Track);
+	TDelegate OnAssetEnterPressed = FOnAssetEnterPressed::CreateRaw(this, &FSubtitlesTrackEditor::OnAssetEnterPressed, Track);
+	return UE::Sequencer::MakeAddButton(
+		LOCTEXT("SubtitleText", "Subtitle")
+		, FOnGetContent::CreateSP(this, &FSubtitlesTrackEditor::BuildSubMenu, MoveTemp(OnAssetSelected), MoveTemp(OnAssetEnterPressed))
+		, Params.ViewModel);
+}
+
+bool FSubtitlesTrackEditor::HandleAssetAdded(UObject* Asset, const FGuid& TargetObjectGuid)
+{
+	if (Asset->IsA<USubtitleAssetUserData>())
+	{
+		auto Subtitle = Cast<USubtitleAssetUserData>(Asset);
+		UMovieSceneSubtitlesTrack* DummyTrack = nullptr;
+		const FScopedTransaction Transaction(LOCTEXT("AddSubtitle_Transaction", "Add Subtitle"));
+		if (TargetObjectGuid.IsValid())
+		{
+			TArray<TWeakObjectPtr<>> OutObjects;
+			for (TWeakObjectPtr<> Object : GetSequencer()->FindObjectsInCurrentSequence(TargetObjectGuid))
+			{
+				if (Object.IsValid())
+				{
+					OutObjects.Add(Object);
+				}
+			}
+			AnimatablePropertyChanged(FOnKeyProperty::CreateRaw(this, &FSubtitlesTrackEditor::AddNewAttachedSubtitle, Subtitle, DummyTrack, OutObjects));
+		}
+		else
+		{
+			int32 RowIndex = INDEX_NONE;
+			AnimatablePropertyChanged(FOnKeyProperty::CreateRaw(this, &FSubtitlesTrackEditor::AddNewSubtitle, Subtitle, DummyTrack, RowIndex));
+		}
+		return true;
+	}
+	return false;
+}
+
+TSharedRef<SWidget> FSubtitlesTrackEditor::BuildSubMenu(FOnAssetSelected OnAssetSelected, FOnAssetEnterPressed OnAssetEnterPressed)
+{
+	TSharedPtr<ISequencer> SequencerPtr = GetSequencer();
+	UMovieSceneSequence* Sequence = SequencerPtr.IsValid() ? SequencerPtr->GetFocusedMovieSceneSequence() : nullptr;
+
+	FAssetRegistryModule& AssetRegistryModule = FModuleManager::LoadModuleChecked<FAssetRegistryModule>(TEXT("AssetRegistry"));
+	TArray<FTopLevelAssetPath> ClassNames;
+	ClassNames.Add(USubtitleAssetUserData::StaticClass()->GetClassPathName());
+	TSet<FTopLevelAssetPath> DerivedClassNames;
+	AssetRegistryModule.Get().GetDerivedClassNames(ClassNames, TSet<FTopLevelAssetPath>(), DerivedClassNames);
+
+	FMenuBuilder MenuBuilder(true, nullptr);
+
+	FAssetPickerConfig AssetPickerConfig;
+	{
+		AssetPickerConfig.OnAssetSelected = OnAssetSelected;
+		AssetPickerConfig.OnAssetEnterPressed = OnAssetEnterPressed;
+		AssetPickerConfig.bAllowNullSelection = false;
+		AssetPickerConfig.bAddFilterUI = true;
+		AssetPickerConfig.bShowTypeInColumnView = false;
+		AssetPickerConfig.InitialAssetViewType = EAssetViewType::List;
+		for (FTopLevelAssetPath ClassName : DerivedClassNames)
+		{
+			AssetPickerConfig.Filter.ClassPaths.Add(ClassName);
+		}
+		AssetPickerConfig.SaveSettingsName = TEXT("SequencerAssetPicker");
+		AssetPickerConfig.AdditionalReferencingAssets.Add(FAssetData(Sequence));
+	}
+
+	FContentBrowserModule& ContentBrowserModule = FModuleManager::Get().LoadModuleChecked<FContentBrowserModule>(TEXT("ContentBrowser"));
+
+	const float WidthOverride = SequencerPtr.IsValid() ? SequencerPtr->GetSequencerSettings()->GetAssetBrowserWidth() : 500.f;
+	const float HeightOverride = SequencerPtr.IsValid() ? SequencerPtr->GetSequencerSettings()->GetAssetBrowserHeight() : 400.f;
+
+	TSharedPtr<SBox> MenuEntry = SNew(SBox)
+		.WidthOverride(WidthOverride)
+		.HeightOverride(HeightOverride)
+		[
+			ContentBrowserModule.Get().CreateAssetPicker(AssetPickerConfig)
+		];
+
+	MenuBuilder.AddWidget(MenuEntry.ToSharedRef(), FText::GetEmpty(), true);
+
+	return MenuBuilder.MakeWidget();
+}
+
+void FSubtitlesTrackEditor::OnAssetSelected(const FAssetData& AssetData, UMovieSceneTrack* Track)
+{
+	FSlateApplication::Get().DismissAllMenus();
+
+	const UObject* SelectedObject = AssetData.GetAsset();
+	if (!SelectedObject)
+	{
+		return;
+	}
+
+	const USubtitleAssetUserData* NewAsset = CastChecked<const USubtitleAssetUserData>(AssetData.GetAsset());
+	if (!NewAsset)
+	{
+		return;
+	}
+	const FScopedTransaction Transaction(NSLOCTEXT("Sequencer", "AddSubtitle_Transaction", "Add Subtitle"));
+
+	UMovieSceneSubtitlesTrack* SubtitlesTrack = Cast<UMovieSceneSubtitlesTrack>(Track);
+	check(SubtitlesTrack);
+	SubtitlesTrack->Modify();
+
+	if (TSharedPtr<ISequencer> SequencerPin = GetSequencer())
+	{
+		UMovieSceneSection* NewSection = SubtitlesTrack->AddNewSubtitle(*NewAsset, SequencerPin->GetLocalTime().Time.FrameNumber);
+
+		SequencerPin->EmptySelection();
+		SequencerPin->SelectSection(NewSection);
+		SequencerPin->ThrobSectionSelection();
+
+		SequencerPin->NotifyMovieSceneDataChanged(EMovieSceneDataChangeType::MovieSceneStructureItemAdded);
+	}
+}
+
+void FSubtitlesTrackEditor::OnAssetEnterPressed(const TArray<FAssetData>& AssetData, UMovieSceneTrack* Track)
+{
+	if (!AssetData.IsEmpty())
+	{
+		OnAssetSelected(AssetData[0].GetAsset(), Track);
+	}
+}
+
+TSharedRef<ISequencerSection> FSubtitlesTrackEditor::MakeSectionInterface(UMovieSceneSection& SectionObject, UMovieSceneTrack& Track, FGuid ObjectBinding)
+{
+	check(SupportsType(SectionObject.GetOuter()->GetClass()));
+	return MakeShareable(new FSubtitleSequencerSection(SectionObject));
+}
+
+bool FSubtitlesTrackEditor::SupportsType(TSubclassOf<UMovieSceneTrack> Type) const
+{
+	return Type == UMovieSceneSubtitlesTrack::StaticClass();
+}
+
+bool FSubtitlesTrackEditor::SupportsSequence(UMovieSceneSequence* InSequence) const
+{
+	ETrackSupport TrackSupported = InSequence ? InSequence->IsTrackSupported(UMovieSceneSubtitlesTrack::StaticClass()) : ETrackSupport::Default;
+
+	if (TrackSupported == ETrackSupport::NotSupported)
+	{
+		return false;
+	}
+
+	return (InSequence && InSequence->IsA(ULevelSequence::StaticClass())) || TrackSupported == ETrackSupport::Supported;
+}
+
+bool FSubtitlesTrackEditor::OnAllowDrop(const FDragDropEvent& DragDropEvent, FSequencerDragDropParams& DragDropParams)
+{
+	DragDropValidInfo DragDropInfo;
+	if (!IsDragDropValid(DragDropEvent, DragDropParams, DragDropInfo))
+	{
+		return false;
+	}
+	for (const FAssetData& AssetData : DragDropInfo.DragDropOp->GetAssets())
+	{
+		if (!MovieSceneToolHelpers::IsValidAsset(DragDropInfo.FocusedSequence, AssetData))
+		{
+			continue;
+		}
+		if (USubtitleAssetUserData* SubtitleAsset = Cast<USubtitleAssetUserData>(AssetData.GetAsset()))
+		{
+			FFrameRate TickResolution = DragDropInfo.SequencerPtr->GetFocusedTickResolution();
+			FFrameNumber LengthInFrames = TickResolution.AsFrameNumber(SubtitleAsset->GetMaximumDuration());
+			DragDropParams.FrameRange = TRange<FFrameNumber>(DragDropParams.FrameNumber, DragDropParams.FrameNumber + LengthInFrames);
+			return true;
+		}
+	}
+	return false;
+}
+
+FReply FSubtitlesTrackEditor::OnDrop(const FDragDropEvent& DragDropEvent, const FSequencerDragDropParams& DragDropParams)
+{
+	DragDropValidInfo DragDropInfo;
+	if (!IsDragDropValid(DragDropEvent, DragDropParams, DragDropInfo))
+	{
+		return FReply::Unhandled();
+	}
+
+	UMovieSceneSubtitlesTrack* SubtitlesTrack = Cast<UMovieSceneSubtitlesTrack>(DragDropParams.Track);
+
+	const FScopedTransaction Transaction(LOCTEXT("DropAssets", "Drop Assets"));
+
+	FMovieSceneTrackEditor::BeginKeying(DragDropParams.FrameNumber);
+
+	bool bAnyDropped = false;
+	for (const FAssetData& AssetData : DragDropInfo.DragDropOp->GetAssets())
+	{
+		if (!MovieSceneToolHelpers::IsValidAsset(DragDropInfo.FocusedSequence, AssetData))
+		{
+			continue;
+		}
+
+		USubtitleAssetUserData* SubtitleAsset = Cast<USubtitleAssetUserData>(AssetData.GetAsset());
+
+		if (SubtitleAsset)
+		{
+			if (DragDropParams.TargetObjectGuid.IsValid())
+			{
+				TArray<TWeakObjectPtr<>> OutObjects;
+				for (TWeakObjectPtr<> Object : DragDropInfo.SequencerPtr->FindObjectsInCurrentSequence(DragDropParams.TargetObjectGuid))
+				{
+					OutObjects.Add(Object);
+				}
+
+				AnimatablePropertyChanged(FOnKeyProperty::CreateRaw(this, &FSubtitlesTrackEditor::AddNewAttachedSubtitle, SubtitleAsset, SubtitlesTrack, OutObjects));
+			}
+			else
+			{
+				AnimatablePropertyChanged(FOnKeyProperty::CreateRaw(this, &FSubtitlesTrackEditor::AddNewSubtitle, SubtitleAsset, SubtitlesTrack, DragDropParams.RowIndex));
+			}
+
+			bAnyDropped = true;
+		}
+	}
+
+	FMovieSceneTrackEditor::EndKeying();
+
+	return bAnyDropped ? FReply::Handled() : FReply::Unhandled();
+}
+
+/** Delegate for AnimatablePropertyChanged in HandleAssetAdded and OnDrop for adding subtitles.  */
+FKeyPropertyResult FSubtitlesTrackEditor::AddNewSubtitle(FFrameNumber KeyTime, USubtitleAssetUserData* Subtitle, UMovieSceneSubtitlesTrack* InDestinationTrack, int32 RowIndex)
+{
+	FKeyPropertyResult KeyPropertyResult;
+
+	UMovieScene* FocusedMovieScene = GetFocusedMovieScene();
+	if (FocusedMovieScene->IsReadOnly())
+	{
+		return KeyPropertyResult;
+	}
+
+	// Undo tracking
+	FocusedMovieScene->Modify();
+
+	FFindOrCreateRootTrackResult<UMovieSceneSubtitlesTrack> TrackResult;
+	TrackResult.Track = InDestinationTrack;
+
+	// Make a new track if a valid one isn't provided.
+	// ::HandleAssetAdded only passes a nullptr, but ::OnDrop should already have a valid destination track (if not, it's safe to create a new one).
+	if (!InDestinationTrack)
+	{
+		TrackResult = FindOrCreateRootTrack<UMovieSceneSubtitlesTrack>();
+		InDestinationTrack = TrackResult.Track;
+	}
+
+	if (ensure(InDestinationTrack))
+	{
+		InDestinationTrack->Modify();
+
+		UMovieSceneSection* NewSection = InDestinationTrack->AddNewSubtitleOnRow(*Subtitle, KeyTime, RowIndex);
+
+		// If a new track was created earlier, set display names and initialize it now.
+		if (TrackResult.bWasCreated)
+		{
+			InDestinationTrack->SetDisplayName(LOCTEXT("SubtitleTrackName", "Subtitle"));
+
+			if (GetSequencer().IsValid())
+			{
+				GetSequencer()->OnAddTrack(InDestinationTrack, FGuid());
+			}
+		}
+
+		KeyPropertyResult.bTrackModified = true;
+		KeyPropertyResult.SectionsCreated.Add(NewSection);
+	}
+
+	return KeyPropertyResult;
+}
+
+/** Delegate for AnimatablePropertyChanged in HandleAssetAdded and OnDrop for attached subtitles.  */
+FKeyPropertyResult FSubtitlesTrackEditor::AddNewAttachedSubtitle(FFrameNumber KeyTime, class USubtitleAssetUserData* Subtitle, UMovieSceneSubtitlesTrack* InDestinationTrack, TArray<TWeakObjectPtr<UObject>> ObjectsToAttachTo)
+{
+	FKeyPropertyResult KeyPropertyResult;
+
+	for (int32 ObjectIndex = 0; ObjectIndex < ObjectsToAttachTo.Num(); ++ObjectIndex)
+	{
+		UObject* Object = ObjectsToAttachTo[ObjectIndex].Get();
+
+		FFindOrCreateHandleResult HandleResult = FindOrCreateHandleToObject(Object);
+		FGuid ObjectHandle = HandleResult.Handle;
+		KeyPropertyResult.bHandleCreated |= HandleResult.bWasCreated;
+
+		if (ObjectHandle.IsValid())
+		{
+			FFindOrCreateTrackResult TrackResult;
+			TrackResult.Track = InDestinationTrack;
+
+			// Make a new track if a valid one isn't provided.
+			// ::HandleAssetAdded only passes a nullptr, but ::OnDrop should already have a valid destination track (if not, it's safe to create a new one).
+			if (!InDestinationTrack)
+			{
+				TrackResult = FindOrCreateTrackForObject(ObjectHandle, UMovieSceneSubtitlesTrack::StaticClass());
+				InDestinationTrack = Cast<UMovieSceneSubtitlesTrack>(TrackResult.Track);
+			}
+
+			KeyPropertyResult.bTrackCreated |= TrackResult.bWasCreated;
+
+			if (ensure(InDestinationTrack))
+			{
+				InDestinationTrack->Modify();
+
+				UMovieSceneSection* NewSection = InDestinationTrack->AddNewSubtitle(*Subtitle, KeyTime);
+				InDestinationTrack->SetDisplayName(LOCTEXT("SubtitleTrackName", "Subtitle"));
+				KeyPropertyResult.bTrackModified = true;
+				KeyPropertyResult.SectionsCreated.Add(NewSection);
+
+				GetSequencer()->EmptySelection();
+				GetSequencer()->SelectSection(NewSection);
+				GetSequencer()->ThrobSectionSelection();
+			}
+		}
+	}
+
+	return KeyPropertyResult;
+}
+#undef LOCTEXT_NAMESPACE
diff --git a/Engine/Plugins/SubtitlesAndClosedCaptions/Source/SubtitlesAndClosedCaptionsEditor/Private/SubtitlesTrackEditor.h b/Engine/Plugins/SubtitlesAndClosedCaptions/Source/SubtitlesAndClosedCaptionsEditor/Private/SubtitlesTrackEditor.h
new file mode 100644
--- /dev/null
+++ b/Engine/Plugins/SubtitlesAndClosedCaptions/Source/SubtitlesAndClosedCaptionsEditor/Private/SubtitlesTrackEditor.h
@@ -0,0 +1,76 @@
+// Copyright Epic Games, Inc. All Rights Reserved.
+
+#pragma once
+
+#include "MovieSceneTrackEditor.h"
+
+#include "ContentBrowserDelegates.h"
+
+struct FGuid;
+class FAssetDragDropOp;
+class FReply;
+class UMovieSceneTrack;
+class UMovieSceneSubtitlesTrack;
+class USubtitleAssetUserData;
+
+struct DragDropValidInfo
+{
+	TSharedPtr<FDragDropOperation> Operation = nullptr;
+	TSharedPtr<ISequencer> SequencerPtr = nullptr;
+	UMovieSceneSequence* FocusedSequence = nullptr;
+	TSharedPtr<FAssetDragDropOp> DragDropOp = nullptr;
+};
+
+/**
+ * Tools for subtitles tracks
+ */
+class FSubtitlesTrackEditor
+	: public FMovieSceneTrackEditor
+{
+public:
+
+	/**
+	 * Constructor
+	 *
+	 * @param InSequencer The sequencer instance to be used by this tool
+	 */
+	FSubtitlesTrackEditor(TSharedRef<ISequencer> InSequencer);
+
+	/**
+	 * Creates an instance of this class.  Called by a sequencer
+	 *
+	 * @param OwningSequencer The sequencer instance to be used by this tool
+	 * @return The new instance of this class
+	 */
+	static TSharedRef<ISequencerTrackEditor> CreateTrackEditor(TSharedRef<ISequencer> OwningSequencer);
+
+public:
+
+	// ISequencerTrackEditor interface
+	virtual FText GetDisplayName() const override;
+	virtual void BuildAddTrackMenu(FMenuBuilder& MenuBuilder) override;
+	virtual TSharedPtr<SWidget> BuildOutlinerEditWidget(const FGuid& ObjectBinding, UMovieSceneTrack* Track, const FBuildEditWidgetParams& Params) override;
+	virtual bool HandleAssetAdded(UObject* Asset, const FGuid& TargetObjectGuid) override;
+	virtual TSharedRef<ISequencerSection> MakeSectionInterface(UMovieSceneSection& SectionObject, UMovieSceneTrack& Track, FGuid ObjectBinding) override;
+	virtual bool SupportsType(TSubclassOf<UMovieSceneTrack> Type) const override;
+	virtual bool SupportsSequence(UMovieSceneSequence* InSequence) const override;
+	virtual bool OnAllowDrop(const FDragDropEvent& DragDropEvent, FSequencerDragDropParams& DragDropParams) override;
+	virtual FReply OnDrop(const FDragDropEvent& DragDropEvent, const FSequencerDragDropParams& DragDropParams) override;
+
+protected:
+	/** Delegate for AnimatablePropertyChanged in HandleAssetAdded for subtitles */
+	FKeyPropertyResult AddNewSubtitle(FFrameNumber KeyTime, class USubtitleAssetUserData* Subtitle, UMovieSceneSubtitlesTrack* InDestinationTrack, int32 RowIndex);
+	/** Delegate for AnimatablePropertyChanged in HandleAssetAdded for attached subtitles */
+	FKeyPropertyResult AddNewAttachedSubtitle(FFrameNumber KeyTime, class USubtitleAssetUserData* Subtitle, UMovieSceneSubtitlesTrack* InDestinationTrack, TArray<TWeakObjectPtr<UObject>> ObjectsToAttachTo);
+
+
+
+private:
+	TSharedRef<SWidget> BuildSubMenu(FOnAssetSelected OnAssetSelected, FOnAssetEnterPressed OnAssetEnterPressed);
+	void OnAssetSelected(const FAssetData& AssetData, UMovieSceneTrack* Track);
+	void OnAssetEnterPressed(const TArray<FAssetData>& AssetData, UMovieSceneTrack* Track);
+	void HandleAddMenuEntryExecute();
+
+	// Helper function: packages up and null-checks common info used by OnAllowDrop and OnDrop
+	bool IsDragDropValid(const FDragDropEvent& DragDropEvent, const FSequencerDragDropParams& DragDropParams, DragDropValidInfo& OutDragDropInfo);
+};
diff --git a/Engine/Plugins/SubtitlesAndClosedCaptions/Source/SubtitlesAndClosedCaptionsEditor/Public/SubtitleAssetActionUtility.h b/Engine/Plugins/SubtitlesAndClosedCaptions/Source/SubtitlesAndClosedCaptionsEditor/Public/SubtitleAssetActionUtility.h
new file mode 100644
--- /dev/null
+++ b/Engine/Plugins/SubtitlesAndClosedCaptions/Source/SubtitlesAndClosedCaptionsEditor/Public/SubtitleAssetActionUtility.h
@@ -0,0 +1,22 @@
+// Copyright Epic Games, Inc. All Rights Reserved.
+
+#pragma once
+#include "CoreMinimal.h"
+#include "AssetActionUtility.h"
+
+#include "SubtitleAssetActionUtility.generated.h"
+
+// Asset Action Utility class for any editor right-click actions relevant to subtitles.
+UCLASS()
+class USubtitleAssetActionUtility : public UAssetActionUtility
+{
+	GENERATED_BODY()
+
+public:
+	USubtitleAssetActionUtility();
+
+	UFUNCTION(CallInEditor, Category="Subtitles")
+	void ConvertBasicOverlaysToSubtitles();
+
+
+};
\ No newline at end of file
diff --git a/Engine/Plugins/SubtitlesAndClosedCaptions/Source/SubtitlesAndClosedCaptionsEditor/Public/SubtitleMigrationUtility.h b/Engine/Plugins/SubtitlesAndClosedCaptions/Source/SubtitlesAndClosedCaptionsEditor/Public/SubtitleMigrationUtility.h
new file mode 100644
--- /dev/null
+++ b/Engine/Plugins/SubtitlesAndClosedCaptions/Source/SubtitlesAndClosedCaptionsEditor/Public/SubtitleMigrationUtility.h
@@ -0,0 +1,30 @@
+// Copyright Epic Games, Inc. All Rights Reserved.
+
+#pragma once
+#include "CoreMinimal.h"
+#include "AssetActionUtility.h"
+
+#include "SubtitleMigrationUtility.generated.h"
+
+// Asset Action Utility class for migrating subtitles
+// Associated asset SubtitleMigrationAction.uasset required to show up in editor.
+UCLASS()
+class USubtitleMigrationUtility : public UAssetActionUtility
+{
+	GENERATED_BODY()
+
+private:
+	const static FString UndoContext;
+
+public:
+	USubtitleMigrationUtility();
+
+	UFUNCTION(CallInEditor, Category = "Subtitles")
+	void AddBlankSubtitle();
+
+	UFUNCTION(CallInEditor, Category = "Subtitles")
+	void RemoveLegacySubtitles();
+
+	UFUNCTION(CallInEditor, Category = "Subtitles")
+	void ConvertLegacySubtitles();
+};
diff --git a/Engine/Plugins/SubtitlesAndClosedCaptions/Source/SubtitlesAndClosedCaptionsEditor/Public/SubtitlesAndClosedCaptionsEditorModule.h b/Engine/Plugins/SubtitlesAndClosedCaptions/Source/SubtitlesAndClosedCaptionsEditor/Public/SubtitlesAndClosedCaptionsEditorModule.h
new file mode 100644
--- /dev/null
+++ b/Engine/Plugins/SubtitlesAndClosedCaptions/Source/SubtitlesAndClosedCaptionsEditor/Public/SubtitlesAndClosedCaptionsEditorModule.h
@@ -0,0 +1,16 @@
+// Copyright Epic Games, Inc. All Rights Reserved.
+#pragma once
+
+#include "Internationalization/Text.h"
+#include "Modules/ModuleInterface.h"
+
+class ISubtitlesAndClosedCaptionsEditorModule : public IModuleInterface
+{
+public:
+	static FText GetAssetTypeCategory()
+	{
+		return AssetTypeCategory;
+	}
+protected:
+	static FText AssetTypeCategory;
+};
diff --git a/Engine/Plugins/SubtitlesAndClosedCaptions/Source/SubtitlesAndClosedCaptionsEditor/SubtitlesAndClosedCaptionsEditor.Build.cs b/Engine/Plugins/SubtitlesAndClosedCaptions/Source/SubtitlesAndClosedCaptionsEditor/SubtitlesAndClosedCaptionsEditor.Build.cs
new file mode 100644
--- /dev/null
+++ b/Engine/Plugins/SubtitlesAndClosedCaptions/Source/SubtitlesAndClosedCaptionsEditor/SubtitlesAndClosedCaptionsEditor.Build.cs
@@ -0,0 +1,42 @@
+// Copyright Epic Games, Inc. All Rights Reserved.
+
+namespace UnrealBuildTool.Rules
+{
+	public class SubtitlesAndClosedCaptionsEditor : ModuleRules
+	{
+		public SubtitlesAndClosedCaptionsEditor(ReadOnlyTargetRules Target) : base(Target)
+		{
+			PCHUsage = PCHUsageMode.UseExplicitOrSharedPCHs;
+
+			PrivateDependencyModuleNames.AddRange
+			(
+				new string[]
+				{
+					"AssetTools",
+					"Blutility",
+					"LevelSequence",
+					"MovieScene",
+					"MovieSceneTools",
+					"Overlay",
+					"Sequencer",
+					"SequencerCore",
+					"Slate",
+					"SlateCore",
+					"SubtitlesAndClosedCaptions",
+				}
+			);
+
+			PublicDependencyModuleNames.AddRange
+			(
+				new string[]
+				{
+					"AssetDefinition",
+					"Core",
+					"CoreUObject",
+					"Engine",
+					"UnrealEd"
+				}
+			);
+		}
+	}
+}
diff --git a/Engine/Plugins/SubtitlesAndClosedCaptions/SubtitlesAndClosedCaptions.uplugin b/Engine/Plugins/SubtitlesAndClosedCaptions/SubtitlesAndClosedCaptions.uplugin
new file mode 100644
--- /dev/null
+++ b/Engine/Plugins/SubtitlesAndClosedCaptions/SubtitlesAndClosedCaptions.uplugin
@@ -0,0 +1,30 @@
+{
+	"FileVersion": 3,
+	"Version": 1000,
+	"VersionName": "1.0.0.0",
+	"FriendlyName": "Subtitles and Closed Captions",
+	"Description": "Standalone plugin for displaying Subtitles and Closed Captions",
+	"Category": "Audio",
+	"CreatedBy": "Epic Games, Inc.",
+	"CreatedByURL": "https://epicgames.com/",
+	"DocsURL": "",
+	"MarketplaceURL": "",
+	"SupportURL": "",
+	"EnabledByDefault": false,
+	"CanContainContent": true,
+	"IsBetaVersion": true,
+	"IsExperimentalVersion": false,
+	"Installed": false,
+	"Modules": [
+		{
+			"Name": "SubtitlesAndClosedCaptions",
+			"Type": "Runtime",
+			"LoadingPhase": "Default"
+		},
+		{
+			"Name": "SubtitlesAndClosedCaptionsEditor",
+			"Type": "UncookedOnly",
+			"LoadingPhase": "Default"
+		}
+	]
+}
