[Subtitles] Display subtitles in-editor, even without PIE, using UViewportWidgetOverlay.
#Synced-CL 51780721

diff --git a/Engine/Plugins/SubtitlesAndClosedCaptions/Source/SubtitlesAndClosedCaptions/Private/SubtitlesSubsystem.cpp b/Engine/Plugins/SubtitlesAndClosedCaptions/Source/SubtitlesAndClosedCaptions/Private/SubtitlesSubsystem.cpp
--- a/Engine/Plugins/SubtitlesAndClosedCaptions/Source/SubtitlesAndClosedCaptions/Private/SubtitlesSubsystem.cpp
+++ b/Engine/Plugins/SubtitlesAndClosedCaptions/Source/SubtitlesAndClosedCaptions/Private/SubtitlesSubsystem.cpp
@@ -232,13 +232,21 @@ void USubtitlesSubsystem::StopAllSubtitles()
 	DelayedSubtitles.Empty();
 
 	// Clear the widget's display
-	if (IsValid(SubtitleWidget))
+	if (IsValid(ViewportWidget))
 	{
-		SubtitleWidget->StopDisplayingAllSubtitles();
+		TWeakObjectPtr<USubtitleWidget> SubtitleWidget = GetActiveSubtitleWidget();
+		if (SubtitleWidget.IsValid())
+		{
+			SubtitleWidget->StopDisplayingAllSubtitles();
+		}
+		else
+		{
+			UE_LOGFMT(LogSubtitlesAndClosedCaptions, Warning, "Can't remove subtitles because there isn't a valid UMG widget.");
+		}
 	}
-	else 
+	else
 	{
-		UE_LOGFMT(LogSubtitlesAndClosedCaptions, Warning, "Can't remove subtitles because there isn't a valid UMG widget.");
+		UE_LOGFMT(LogSubtitlesAndClosedCaptions, Warning, "Can't remove subtitles because the Viewport Widget couldn't be initialized.");
 	}
 }
 
@@ -275,19 +283,27 @@ void USubtitlesSubsystem::RemoveActiveSubtitle(const FSubtitleAssetData& Subtitl
 	ActiveSubtitles.SetNum(FirstRemovedIndex);
 
 	// Stop Displaying the removed subtitle and display a newly-most-relevant one if applicable.
-	if (bSuccessfullyRemoved && IsValid(SubtitleWidget))
+	if (bSuccessfullyRemoved && IsValid(ViewportWidget))
 	{
-		SubtitleWidget->StopDisplayingSubtitle(Subtitle);
+		TWeakObjectPtr<USubtitleWidget> SubtitleWidget = GetActiveSubtitleWidget();
+		if (SubtitleWidget.IsValid())
+		{
+			SubtitleWidget->StopDisplayingSubtitle(Subtitle);
 
-		if (!ActiveSubtitles.IsEmpty())
+			if (!ActiveSubtitles.IsEmpty())
+			{
+				const FSubtitleAssetData& HighestPrioritySubtitle = ActiveSubtitles[0].Subtitle;
+				SubtitleWidget->StartDisplayingSubtitle(HighestPrioritySubtitle);
+			}
+		}
+		else
 		{
-			const FSubtitleAssetData& HighestPrioritySubtitle = ActiveSubtitles[0].Subtitle;
-			SubtitleWidget->StartDisplayingSubtitle(HighestPrioritySubtitle);
+			UE_LOGFMT(LogSubtitlesAndClosedCaptions, Warning, "Can't remove the next subtitle because there isn't a valid UMG widget.");
 		}
 	}
-	else if (!IsValid(SubtitleWidget))
+	else if (!IsValid(ViewportWidget))
 	{
-		UE_LOGFMT(LogSubtitlesAndClosedCaptions, Warning, "Can't remove subtitles because there isn't a valid UMG widget.");
+		UE_LOGFMT(LogSubtitlesAndClosedCaptions, Warning, "Can't remove subtitles because the Viewport Widget couldn't be initialized.");
 	}
 
 	// Remove delayed-start subtitles that haven't been queued yet that also use this asset.
@@ -302,21 +318,30 @@ void USubtitlesSubsystem::RemoveActiveSubtitle(const FSubtitleAssetData& Subtitl
 	DelayedSubtitles.SetNum(FirstRemovedActiveIndex);
 }
 
+TWeakObjectPtr<USubtitleWidget> USubtitlesSubsystem::GetActiveSubtitleWidget() const
+{
+	if (IsValid(ViewportWidget))
+	{
+		return Cast<USubtitleWidget>(ViewportWidget->GetWidget());
+	}
+	return nullptr;
+}
+
 bool USubtitlesSubsystem::TryCreateUMGWidgetFromAsset(const TSoftClassPtr<USubtitleWidget>& WidgetToUse)
 {
-	// Synchronous load: called directly from ::Initialize when the World (and this subsystem) is created.
-	WidgetToUse.LoadSynchronous();
-	if (WidgetToUse.IsValid())
+	if (IsValid(ViewportWidget))
 	{
-		SubtitleWidget = CreateWidget<USubtitleWidget>(GetWorld(), WidgetToUse.Get());
-		if (IsValid(SubtitleWidget))
+
+		// Synchronous load: called directly from ::Initialize when the World (and this subsystem) is created.
+		WidgetToUse.LoadSynchronous();
+		if (WidgetToUse.IsValid())
 		{
-			SubtitleWidget->SetFlags(RF_Transient);
+			ViewportWidget->SetWidgetClass(WidgetToUse.Get());
+			bInitializedWidget = false;	// The new widget will need to be displayed.
+			return true;
 		}
-		bInitializedWidget = false;	// The new widget will need to be added to the viewport.
 	}
-
-	return IsValid(SubtitleWidget);
+	return false;
 }
 
 bool USubtitlesSubsystem::TryCreateUMGWidget()
@@ -325,18 +350,26 @@ bool USubtitlesSubsystem::TryCreateUMGWidget()
 	const USubtitlesSettings* Settings = GetDefault<USubtitlesSettings>();
 	check(Settings != nullptr);
 	const TSoftClassPtr<USubtitleWidget>& WidgetToUse = Settings->GetWidgetSoftClassPtr();
-	
+
+	if (!IsValid(ViewportWidget))
+	{
+		ViewportWidget = NewObject<UViewportWidgetOverlay>(this);
+		if (!IsValid(ViewportWidget))
+		{
+			UE_LOGFMT(LogSubtitlesAndClosedCaptions, Error, "Couldn't initialize the Subtitles Viewport Widget. Subtitles won't be displayed.");
+			return false;
+		}
+
+		ViewportWidget->SetDisplayTypes(EViewportWidgetOverlay_DisplayType::Viewport, EViewportWidgetOverlay_DisplayType::Viewport, EViewportWidgetOverlay_DisplayType::Viewport);
+	}
+
 	if (!TryCreateUMGWidgetFromAsset(WidgetToUse))
 	{
 		// Fallback to default widget (not set by user):
 		const TSoftClassPtr<USubtitleWidget>& WidgetToUseDefault = Settings->GetWidgetDefaultSoftClassPtr();
 		if (WidgetToUseDefault.IsValid())
 		{
-			SubtitleWidget = CreateWidget<USubtitleWidget>(GetWorld(), WidgetToUseDefault.Get());
-			if (IsValid(SubtitleWidget))
-			{
-				SubtitleWidget->SetFlags(RF_Transient);
-			}
+			ViewportWidget->SetWidgetClass(WidgetToUseDefault.Get());
 			// Log the fallback as a warning as seeing the default widget probably isn't intended.
 			UE_LOGFMT(LogSubtitlesAndClosedCaptions, Warning, "A valid Subtitle Widget wasn't provided. The plugin's default widget will be used as a fallback.");
 		}
@@ -347,27 +380,35 @@ bool USubtitlesSubsystem::TryCreateUMGWidget()
 	}
 
 	bInitializedWidget = false;
-	return IsValid(SubtitleWidget);
+	return IsValid(ViewportWidget);
 }
 
 void USubtitlesSubsystem::UpdateWidgetData()
 {
 	// Update the widget. If it's not valid (eg, destroyed on non-seamless travel), try re-creating it first.
-	if (IsValid(SubtitleWidget) || (bInitializedWidget && TryCreateUMGWidget() && IsValid(SubtitleWidget)))
+	if (IsValid(ViewportWidget) || (bInitializedWidget && TryCreateUMGWidget() && IsValid(ViewportWidget)))
 	{
 		if (!bInitializedWidget)
 		{
-			SubtitleWidget->AddToViewport();
-			SubtitleWidget->SetVisibility(ESlateVisibility::HitTestInvisible);
+			ViewportWidget->Display(GetWorld());
 			bInitializedWidget = true;
 		}
 
-		const FSubtitleAssetData& NewHighestPrioritySubtitle = ActiveSubtitles[0].Subtitle;
-		SubtitleWidget->StartDisplayingSubtitle(NewHighestPrioritySubtitle);
+		TWeakObjectPtr<USubtitleWidget> SubtitleWidget = GetActiveSubtitleWidget();
+		if (SubtitleWidget.IsValid())
+		{
+			check(SubtitleWidget->GetFlags() & RF_Transient);
+			const FSubtitleAssetData& NewHighestPrioritySubtitle = ActiveSubtitles[0].Subtitle;
+			SubtitleWidget->StartDisplayingSubtitle(NewHighestPrioritySubtitle);
+		}
+		else
+		{
+			UE_LOGFMT(LogSubtitlesAndClosedCaptions, Warning, "Can't display subtitles, because there isn't a valid UMG Widget to display it to (check the Subtitle Widget and defaults in your project settings).");
+		}
 	}
 	else
 	{
-		UE_LOGFMT(LogSubtitlesAndClosedCaptions, Warning, "Can't display subtitles because there isn't a valid UMG widget to display it to (check Project Settings).");
+		UE_LOGFMT(LogSubtitlesAndClosedCaptions, Warning, "Can't display subtitles because there isn't a valid Viewport Widget to display it to.");
 	}
 }
 
diff --git a/Engine/Plugins/SubtitlesAndClosedCaptions/Source/SubtitlesAndClosedCaptions/Public/SubtitlesSubsystem.h b/Engine/Plugins/SubtitlesAndClosedCaptions/Source/SubtitlesAndClosedCaptions/Public/SubtitlesSubsystem.h
--- a/Engine/Plugins/SubtitlesAndClosedCaptions/Source/SubtitlesAndClosedCaptions/Public/SubtitlesSubsystem.h
+++ b/Engine/Plugins/SubtitlesAndClosedCaptions/Source/SubtitlesAndClosedCaptions/Public/SubtitlesSubsystem.h
@@ -9,6 +9,7 @@
 #include "Subtitles/SubtitlesAndClosedCaptionsDelegates.h"
 
 #include "SubtitleWidget.h"
+#include "ViewportWidgetOverlay.h"
 
 #include "SubtitlesSubsystem.generated.h"
 
@@ -91,6 +92,8 @@ private:
 	UE_API void AddAndDisplaySubtitle(FActiveSubtitleAssetData& NewActiveSubtitle);
 
 public:
+	UE_API TWeakObjectPtr<USubtitleWidget> GetActiveSubtitleWidget() const;
+
 	// GetActiveSubtitles returns the sorted array of queued subtitles. Higher priorities are earlier in the array.
 	// For equally-prioritized subtitles, the oldest one is first in this array.
 	const TArray<FActiveSubtitleAssetData>& GetActiveSubtitles() const
@@ -98,14 +101,9 @@ public:
 		return ActiveSubtitles;
 	}
 
-	inline TWeakObjectPtr<USubtitleWidget> GetActiveSubtitleWidget()
-	{
-		return SubtitleWidget;
-	}
-
 private:
 	UPROPERTY(Transient)
-	TObjectPtr<USubtitleWidget> SubtitleWidget = nullptr;
+	TObjectPtr<UViewportWidgetOverlay> ViewportWidget = nullptr;
 
 	bool bInitializedWidget = false;
 	UE_API bool TryCreateUMGWidgetFromAsset(const TSoftClassPtr<USubtitleWidget>& WidgetToUse);
diff --git a/Engine/Plugins/SubtitlesAndClosedCaptions/Source/SubtitlesAndClosedCaptions/SubtitlesAndClosedCaptions.Build.cs b/Engine/Plugins/SubtitlesAndClosedCaptions/Source/SubtitlesAndClosedCaptions/SubtitlesAndClosedCaptions.Build.cs
--- a/Engine/Plugins/SubtitlesAndClosedCaptions/Source/SubtitlesAndClosedCaptions/SubtitlesAndClosedCaptions.Build.cs
+++ b/Engine/Plugins/SubtitlesAndClosedCaptions/Source/SubtitlesAndClosedCaptions/SubtitlesAndClosedCaptions.Build.cs
@@ -27,7 +27,8 @@ namespace UnrealBuildTool.Rules
 					"InputCore",
 					"MovieSceneTracks",
 					"MovieScene",
-					"UMG"
+					"UMG",
+					"ViewportWidgetOverlay"
 				}
 			);
 		}
diff --git a/Engine/Plugins/SubtitlesAndClosedCaptions/SubtitlesAndClosedCaptions.uplugin b/Engine/Plugins/SubtitlesAndClosedCaptions/SubtitlesAndClosedCaptions.uplugin
--- a/Engine/Plugins/SubtitlesAndClosedCaptions/SubtitlesAndClosedCaptions.uplugin
+++ b/Engine/Plugins/SubtitlesAndClosedCaptions/SubtitlesAndClosedCaptions.uplugin
@@ -15,6 +15,12 @@
 	"IsBetaVersion": true,
 	"IsExperimentalVersion": false,
 	"Installed": false,
+	"Plugins": [
+		{
+			"Name": "ViewportWidgetOverlay",
+			"Enabled": true
+		}
+	],
 	"Modules": [
 		{
 			"Name": "SubtitlesAndClosedCaptions",
