[Subtitles] Make subtitles automatically play Waves nested in MetaSoundSources

#Synced-CL 51836005

diff --git a/Engine/Plugins/SubtitlesAndClosedCaptions/Source/SubtitlesAndClosedCaptions/Private/SubtitlesSubsystem.cpp b/Engine/Plugins/SubtitlesAndClosedCaptions/Source/SubtitlesAndClosedCaptions/Private/SubtitlesSubsystem.cpp
--- a/Engine/Plugins/SubtitlesAndClosedCaptions/Source/SubtitlesAndClosedCaptions/Private/SubtitlesSubsystem.cpp
+++ b/Engine/Plugins/SubtitlesAndClosedCaptions/Source/SubtitlesAndClosedCaptions/Private/SubtitlesSubsystem.cpp
@@ -4,14 +4,19 @@
 
 #include "Algo/Find.h"
 #include "Algo/RemoveIf.h"
+#include "AssetRegistry/AssetData.h"
+#include "AssetRegistry/AssetRegistryModule.h"
 #include "CanvasItem.h"
 #include "CanvasTypes.h"
 #include "Containers/Ticker.h"
 #include "Engine/Font.h"
 #include "Engine/World.h"
 #include "Engine/Engine.h"
+#include "AssetRegistry/IAssetRegistry.h"
+#include "UObject/SoftObjectPath.h"
 #include "Logging/StructuredLog.h"
 #include "Math/UnrealMathUtility.h"
+#include "Sound/SoundWave.h"
 #include "Styling/CoreStyle.h"
 #include "SubtitlesAndClosedCaptionsModule.h"
 #include "SubtitlesSettings.h"
@@ -29,20 +34,38 @@ static constexpr float InfiniteDuration = FLT_MAX;
 
 void USubtitlesSubsystem::Initialize(FSubsystemCollectionBase& Collection)
 {
+	FAssetRegistryModule& AssetRegistryModule = FModuleManager::LoadModuleChecked<FAssetRegistryModule>(TEXT("AssetRegistry"));
+	AssetRegistry = &AssetRegistryModule.Get();
+
 	BindDelegates();
 }
 
+void USubtitlesSubsystem::Deinitialize()
+{
+	UnbindDelegates();
+}
+
 void USubtitlesSubsystem::BindDelegates()
 {
 	if (TryCreateUMGWidget())
 	{
 		FSubtitlesAndClosedCaptionsDelegates::QueueSubtitle.BindUObject(this, &USubtitlesSubsystem::QueueSubtitle);
+		FSubtitlesAndClosedCaptionsDelegates::QueueSubtitleFromSoundWaveName.BindUObject(this, &USubtitlesSubsystem::QueueSubtitleFromSoundWaveName);
 		FSubtitlesAndClosedCaptionsDelegates::IsSubtitleActive.BindUObject(this, &USubtitlesSubsystem::IsSubtitleActive);
 		FSubtitlesAndClosedCaptionsDelegates::StopSubtitle.BindUObject(this, &USubtitlesSubsystem::StopSubtitle);
 		FSubtitlesAndClosedCaptionsDelegates::StopAllSubtitles.BindUObject(this, &USubtitlesSubsystem::StopAllSubtitles);
 	}
 }
 
+void USubtitlesSubsystem::UnbindDelegates()
+{
+	FSubtitlesAndClosedCaptionsDelegates::QueueSubtitle.Unbind();
+	FSubtitlesAndClosedCaptionsDelegates::QueueSubtitleFromSoundWaveName.Unbind();
+	FSubtitlesAndClosedCaptionsDelegates::IsSubtitleActive.Unbind();
+	FSubtitlesAndClosedCaptionsDelegates::StopSubtitle.Unbind();
+	FSubtitlesAndClosedCaptionsDelegates::StopAllSubtitles.Unbind();
+}
+
 void USubtitlesSubsystem::QueueSubtitle(FQueueSubtitleParameters Params, const ESubtitleTiming Timing)
 {
 	// Don't queue subtitles when the engine doesn't have them enabled or forced off. 
@@ -87,6 +110,28 @@ void USubtitlesSubsystem::QueueSubtitle(FQueueSubtitleParameters Params, const E
 	}
 }
 
+void USubtitlesSubsystem::QueueSubtitleFromSoundWaveName(const FName& SoundName, const FName& PackageName, const float Duration)
+{
+	// Move this work to the game thread.
+	if (IsInGameThread())
+	{
+		FindAndQueueSubtitleFromSoundName(SoundName, PackageName, Duration);
+	}
+	else
+	{
+		ExecuteOnGameThread(
+			TEXT("USubtitlesSubsystem::FindAndQueueSubtitleFromSoundName"),
+			[WeakThis = TWeakObjectPtr<USubtitlesSubsystem>(this), SoundName, PackageName, Duration]()
+			{
+				if (TStrongObjectPtr<USubtitlesSubsystem> ThisPin = WeakThis.Pin())
+				{
+					ThisPin->FindAndQueueSubtitleFromSoundName(SoundName, PackageName, Duration);
+				}
+			}
+		);
+	}
+}
+
 void USubtitlesSubsystem::AddActiveSubtitle(const FSubtitleAssetData& Subtitle, float Duration, const float StartOffset, const ESubtitleTiming Timing)
 {
 	// If the subtitle is already active then update its duration (by removing it then and then re-adding it)
@@ -126,6 +171,40 @@ void USubtitlesSubsystem::AddActiveSubtitle(const FSubtitleAssetData& Subtitle,
 	}
 }
 
+void USubtitlesSubsystem::FindAndQueueSubtitleFromSoundName(const FName& SoundName, const FName& PackageName, const float Duration)
+{
+	if (GEngine == nullptr || !GEngine->bSubtitlesEnabled || GEngine->bSubtitlesForcedOff)
+	{
+		return;
+	}
+
+	if (!ensure(AssetRegistry))
+	{
+		UE_LOGFMT(LogSubtitlesAndClosedCaptions, Error, "Couldn't queue subtitles from {SoundName}; the AssetRegistry module wasn't loaded.", SoundName);
+		return;
+	}
+
+	const FAssetData AssetData = AssetRegistry->GetAssetByObjectPath(FSoftObjectPath(FTopLevelAssetPath(PackageName, SoundName)));
+
+	// These can't be const because HasAssetUserDataOfClass isn't
+	// FastGetAsset does not load the asset. Don't bother loading if it isn't already.
+	// It's already be loaded by the calling MSS, and this is called too often to do any blocking loads.
+	UObject* SoundWaveAsset = AssetData.FastGetAsset(/* bLoad = false */);
+	USoundWave* AsSoundWave = Cast<USoundWave>(SoundWaveAsset);
+
+	if (IsValid(AsSoundWave) && AsSoundWave->HasAssetUserDataOfClass(USubtitleAssetUserData::StaticClass()))
+	{
+		// Using CastChecked because we already determined that the correct asset userdata exists.
+		const UAssetUserData* UserData = AsSoundWave->GetAssetUserDataOfClass(USubtitleAssetUserData::StaticClass());
+		const USubtitleAssetUserData* SubtitleUserData = CastChecked<USubtitleAssetUserData>(UserData);
+
+		for (const FSubtitleAssetData& Subtitle : SubtitleUserData->Subtitles)
+		{
+			FSubtitlesAndClosedCaptionsDelegates::QueueSubtitle.ExecuteIfBound(FQueueSubtitleParameters{ Subtitle, Duration }, ESubtitleTiming::InternallyTimed);
+		}
+	}
+}
+
 void USubtitlesSubsystem::MakeDelayedSubtitleActive(const FSubtitleAssetData& Subtitle, const ESubtitleTiming Timing)
 {
 	FActiveSubtitleAssetData* DelayedSubtitle = Algo::FindByPredicate(DelayedSubtitles, [&Subtitle](const FActiveSubtitleAssetData& DelayedSubtitle) {return DelayedSubtitle.Subtitle == Subtitle; });
diff --git a/Engine/Plugins/SubtitlesAndClosedCaptions/Source/SubtitlesAndClosedCaptions/Public/SubtitlesSubsystem.h b/Engine/Plugins/SubtitlesAndClosedCaptions/Source/SubtitlesAndClosedCaptions/Public/SubtitlesSubsystem.h
--- a/Engine/Plugins/SubtitlesAndClosedCaptions/Source/SubtitlesAndClosedCaptions/Public/SubtitlesSubsystem.h
+++ b/Engine/Plugins/SubtitlesAndClosedCaptions/Source/SubtitlesAndClosedCaptions/Public/SubtitlesSubsystem.h
@@ -16,6 +16,7 @@
 #define UE_API SUBTITLESANDCLOSEDCAPTIONS_API
 
 class FCanvas;
+class IAssetRegistry;
 class UAssetUserData;
 class UFont;
 struct FQueueSubtitleParameters;
@@ -50,6 +51,9 @@ public:
 	// If the subtitle asset has a non-zero StartOffset, it will sit in a delayed-start queue instead of being queued for display.
 	UE_API virtual void QueueSubtitle(FQueueSubtitleParameters Params, const ESubtitleTiming Timing = ESubtitleTiming::InternallyTimed);
 
+	// Metasounds may reference a soundwave, but can only really provide the sound's name from the audio thread.
+	UE_API virtual void QueueSubtitleFromSoundWaveName(const FName& SoundName, const FName& PackageName, const float Duration);
+
 	// Returns true if the given subtitle asset is being displayed.
 	UE_API virtual bool IsSubtitleActive(FSubtitleAssetData Data) const;
 
@@ -65,7 +69,9 @@ public:
 protected:
 
 	UE_API virtual void Initialize(FSubsystemCollectionBase& Collection) override;
+	UE_API virtual void Deinitialize() override;
 	UE_API virtual void BindDelegates();
+	UE_API virtual void UnbindDelegates();
 
 	UPROPERTY()
 	TMap<ESubtitleType, FSlateFontInfo> SubtitleFontInfo;
@@ -81,6 +87,8 @@ protected:
 protected:
 	UE_API virtual void AddActiveSubtitle(const FSubtitleAssetData& Subtitle, float Duration, const float StartOffset, const ESubtitleTiming Timing);
 
+	UE_API virtual void FindAndQueueSubtitleFromSoundName(const FName& SoundName, const FName& PackageName, const float Duration);
+
 	// These need to be UFUNCTIONs for the Timer Delegate bindings.
 	UFUNCTION()
 	UE_API virtual void MakeDelayedSubtitleActive(const FSubtitleAssetData& Subtitle, const ESubtitleTiming Timing);
@@ -106,6 +114,7 @@ private:
 	TObjectPtr<UViewportWidgetOverlay> ViewportWidget = nullptr;
 
 	bool bInitializedWidget = false;
+	IAssetRegistry* AssetRegistry = nullptr;
 	UE_API bool TryCreateUMGWidgetFromAsset(const TSoftClassPtr<USubtitleWidget>& WidgetToUse);
 	UE_API bool TryCreateUMGWidget();
 	UE_API void UpdateWidgetData();
diff --git a/Engine/Source/Runtime/Engine/Classes/Sound/SoundWave.h b/Engine/Source/Runtime/Engine/Classes/Sound/SoundWave.h
--- a/Engine/Source/Runtime/Engine/Classes/Sound/SoundWave.h
+++ b/Engine/Source/Runtime/Engine/Classes/Sound/SoundWave.h
@@ -1592,6 +1592,7 @@ public:
 		, bShouldUseStreamCaching(0)
 		, bLoadingBehaviorOverridden(0)
 		, bHasError(0)
+		, bHasSubtitleAssetUserData(0)
 #if WITH_EDITOR
 		, bLoadedFromCookedData(0)
 #endif //WITH_EDITOR
@@ -1636,6 +1637,7 @@ public:
 	bool IsSeekable() const {  return bIsSeekable; }
 	bool IsRetainingAudio() const {  return FirstChunk.IsValid(); }
 	bool WasLoadingBehaviorOverridden() const { return bLoadingBehaviorOverridden; }
+	bool HasSubtitleAssetUserData() const { return bHasSubtitleAssetUserData; }
 	ESoundWaveLoadingBehavior GetLoadingBehavior() const { return LoadingBehavior; }
 
 	bool IsTemplate() const { return bIsTemplate; }
@@ -1741,6 +1743,7 @@ private:
 	uint8 bShouldUseStreamCaching : 1;
 	uint8 bLoadingBehaviorOverridden : 1;
 	uint8 bHasError : 1;
+	uint8 bHasSubtitleAssetUserData : 1;
 
 #if WITH_EDITOR
 	uint8 bLoadedFromCookedData : 1;
@@ -1796,6 +1799,7 @@ public:
 	ENGINE_API bool IsZerothChunkDataLoaded() const;
 	ENGINE_API bool WasLoadingBehaviorOverridden() const;
 	ENGINE_API bool HasCompressedData(FName Format, ITargetPlatform* TargetPlatform = USoundWave::GetRunningPlatform()) const;
+	ENGINE_API bool HasSubtitleAssetUserData() const;
 
 	ENGINE_API ESoundWaveLoadingBehavior GetLoadingBehavior() const;
 
diff --git a/Engine/Source/Runtime/Engine/Classes/Sound/SoundWaveProxyPlayer.h b/Engine/Source/Runtime/Engine/Classes/Sound/SoundWaveProxyPlayer.h
--- a/Engine/Source/Runtime/Engine/Classes/Sound/SoundWaveProxyPlayer.h
+++ b/Engine/Source/Runtime/Engine/Classes/Sound/SoundWaveProxyPlayer.h
@@ -389,6 +389,9 @@ private:
 	static constexpr float MaxSpeed = 100.0f;
 	
 	bool bFallbackSeekMethodWarningLogged = false;
+
+	// Track how many samples played before queueing another subtitle chunk.
+	int64 Chan0SamplesSinceQueuedSubtitles = 0;
 };
 
 #undef UE_API
\ No newline at end of file
diff --git a/Engine/Source/Runtime/Engine/Private/SoundWave.cpp b/Engine/Source/Runtime/Engine/Private/SoundWave.cpp
--- a/Engine/Source/Runtime/Engine/Private/SoundWave.cpp
+++ b/Engine/Source/Runtime/Engine/Private/SoundWave.cpp
@@ -260,6 +260,7 @@ void FSoundWaveData::InitializeDataFromSoundWave(USoundWave& InWave)
 	bIsSeekable = false; // This is determined by the decoder in CacheRuntimeFormatDependentState.
 	SoundAssetCompressionType = InWave.GetSoundAssetCompressionType();
 	bShouldUseStreamCaching = InWave.ShouldUseStreamCaching();
+	bHasSubtitleAssetUserData = InWave.HasAssetUserDataOfClass(USubtitleAssetUserData::StaticClass());
 
 	// Cache the decoder and any dependent state.
 	const FName FoundFormat = FindRuntimeFormat(InWave);
@@ -5492,6 +5493,16 @@ bool FSoundWaveProxy::HasCompressedData(FName Format, ITargetPlatform* TargetPla
 	return SoundWaveDataPtr->HasCompressedData(Format, TargetPlatform);
 }
 
+bool FSoundWaveProxy::HasSubtitleAssetUserData() const
+{
+	ensure(SoundWaveDataPtr);
+	if (SoundWaveDataPtr)
+	{
+		return SoundWaveDataPtr->HasSubtitleAssetUserData();
+	}
+	return false;
+}
+
 ESoundWaveLoadingBehavior FSoundWaveProxy::GetLoadingBehavior() const
 {
 	check(SoundWaveDataPtr);
diff --git a/Engine/Source/Runtime/Engine/Private/SoundWaveProxyPlayer.cpp b/Engine/Source/Runtime/Engine/Private/SoundWaveProxyPlayer.cpp
--- a/Engine/Source/Runtime/Engine/Private/SoundWaveProxyPlayer.cpp
+++ b/Engine/Source/Runtime/Engine/Private/SoundWaveProxyPlayer.cpp
@@ -7,6 +7,7 @@
 #include "DSP/FloatArrayMath.h"
 #include "Interfaces/IAudioFormat.h"
 #include "ProfilingDebugging/CpuProfilerTrace.h"
+#include "Subtitles/SubtitlesAndClosedCaptionsDelegates.h"
 #include "Templates/SharedPointer.h"
 
 FSoundWaveProxyPlayer::FSoundWaveProxyPlayer(const FSettings& InSettings)
@@ -63,6 +64,8 @@ bool FSoundWaveProxyPlayer::SetSoundWave(FSoundWaveProxyPtr InWaveProxy)
 	ConvertDeinterleaveParams.MonoUpmixMethod = Settings.MonoUpmixMethod;
 	ConvertDeinterleave = Audio::IConvertDeinterleave::Create(ConvertDeinterleaveParams);
 
+	Chan0SamplesSinceQueuedSubtitles = 0;
+
 	return true;
 }
 
@@ -129,6 +132,8 @@ void FSoundWaveProxyPlayer::Reset(const FSettings& NewSettings)
 	NumFramesInWave = -1;
 	DurationSeconds = -1.0f;
 	SortedCuePoints.Reset();
+
+	Chan0SamplesSinceQueuedSubtitles = 0;
 }
 
 TUniquePtr<FSoundWaveProxyPlayer> FSoundWaveProxyPlayer::Create(const FSettings& InSettings)
@@ -355,6 +360,27 @@ void FSoundWaveProxyPlayer::RenderMultiChannelAudio(const Audio::FMultichannelBu
 		}
 		SourceAudio.SetNum(NumSamplesRemaining, EAllowShrinking::No);
 	}
+
+	// Similar to Sequencer regularly re-queuing short chunks of the subtitles as it scrubs, this re-queues them from the sound/package names sent to the SubtitlesSubsystem.
+	constexpr static float kSubtitleDisplayDuration = 0.125f;				// Subtitle display duration, in seconds
+	constexpr static float kSubtitleRequeueDuration = 0.1f;					// How long in seconds to wait before queueing again, slightly shorter than the display duration.
+
+	// Uses the first channel to count samples, so check that there's a channel there first.
+	if (SoundWave->HasSubtitleAssetUserData() && NumChannels > 0)
+	{
+		Chan0SamplesSinceQueuedSubtitles += OutMultiChannelAudio[0].Num();
+		check(Settings.OutputSampleRate > 0);
+
+		const float ElapsedDuration = Chan0SamplesSinceQueuedSubtitles / Settings.OutputSampleRate;
+		if (ElapsedDuration >= kSubtitleRequeueDuration)
+		{
+			Chan0SamplesSinceQueuedSubtitles = 0;
+
+			// Send the sound's name to the subtitle system, where it will look up any associated subtitles and play them.
+			// SubtitlesSubsystem will receive the delegate and send the FNames to the game thread for queueing.
+			FSubtitlesAndClosedCaptionsDelegates::QueueSubtitleFromSoundWaveName.ExecuteIfBound(SoundWave->GetFName(), SoundWave->GetPackageName(), kSubtitleDisplayDuration);
+		}
+	}
 }
 
 int32 FSoundWaveProxyPlayer::GenerateSourceAudio(TArrayView<float> OutAudio)
diff --git a/Engine/Source/Runtime/Engine/Private/Subtitles/SubtitlesAndClosedCaptionsDelegates.cpp b/Engine/Source/Runtime/Engine/Private/Subtitles/SubtitlesAndClosedCaptionsDelegates.cpp
--- a/Engine/Source/Runtime/Engine/Private/Subtitles/SubtitlesAndClosedCaptionsDelegates.cpp
+++ b/Engine/Source/Runtime/Engine/Private/Subtitles/SubtitlesAndClosedCaptionsDelegates.cpp
@@ -4,6 +4,7 @@
 #include "Subtitles/SubtitlesAndClosedCaptionsTypes.h"
 
 TDelegate<void(FQueueSubtitleParameters, const ESubtitleTiming)> FSubtitlesAndClosedCaptionsDelegates::QueueSubtitle;
+TDelegate<void(const FName&, const FName&, const float), FDefaultTSDelegateUserPolicy> FSubtitlesAndClosedCaptionsDelegates::QueueSubtitleFromSoundWaveName;
 TDelegate<bool(FSubtitleAssetData)> FSubtitlesAndClosedCaptionsDelegates::IsSubtitleActive;
 TDelegate<void(FSubtitleAssetData)> FSubtitlesAndClosedCaptionsDelegates::StopSubtitle;
 TDelegate<void()> FSubtitlesAndClosedCaptionsDelegates::StopAllSubtitles;
\ No newline at end of file
diff --git a/Engine/Source/Runtime/Engine/Public/Subtitles/SubtitlesAndClosedCaptionsDelegates.h b/Engine/Source/Runtime/Engine/Public/Subtitles/SubtitlesAndClosedCaptionsDelegates.h
--- a/Engine/Source/Runtime/Engine/Public/Subtitles/SubtitlesAndClosedCaptionsDelegates.h
+++ b/Engine/Source/Runtime/Engine/Public/Subtitles/SubtitlesAndClosedCaptionsDelegates.h
@@ -16,6 +16,8 @@ public:
 	// Have the subtitle subsystem to queue a subtitle to be displayed
 	static UE_API TDelegate<void(FQueueSubtitleParameters, const ESubtitleTiming)> QueueSubtitle;
 
+	static UE_API TDelegate<void(const FName& SoundName, const FName& PackageName, const float Duration), FDefaultTSDelegateUserPolicy> QueueSubtitleFromSoundWaveName;
+
 	static UE_API TDelegate<bool(FSubtitleAssetData)> IsSubtitleActive;
 
 	static UE_API TDelegate<void(FSubtitleAssetData)> StopSubtitle;
