[Subtitles] Only Prevent non-playing DialogueWaves referenced by SoundCues from queueing subtitles.

#Synced-CL 52110155
#Synced-CL 52110155

diff --git a/Engine/Plugins/SubtitlesAndClosedCaptions/Source/SubtitlesAndClosedCaptions/Private/SubtitlesAudioSubsystem.cpp b/Engine/Plugins/SubtitlesAndClosedCaptions/Source/SubtitlesAndClosedCaptions/Private/SubtitlesAudioSubsystem.cpp
--- a/Engine/Plugins/SubtitlesAndClosedCaptions/Source/SubtitlesAndClosedCaptions/Private/SubtitlesAudioSubsystem.cpp
+++ b/Engine/Plugins/SubtitlesAndClosedCaptions/Source/SubtitlesAndClosedCaptions/Private/SubtitlesAudioSubsystem.cpp
@@ -4,13 +4,24 @@
 
 #include "ActiveSound.h"
 #include "Sound/DialogueSoundWaveProxy.h"
-#include "Sound/SoundCue.h"
 #include "Sound/SoundNodeDialoguePlayer.h"
 #include "Subtitles/SubtitlesAndClosedCaptionsTypes.h"
 #include "Subtitles/SubtitlesAndClosedCaptionsDelegates.h"
 
 #include UE_INLINE_GENERATED_CPP_BY_NAME(SubtitlesAudioSubsystem)
 
+void USubtitlesAudioSubsystem::Initialize(FSubsystemCollectionBase& Collection)
+{
+	Super::Initialize(Collection);
+	USoundNodeDialoguePlayer::OnDialogueSoundWaveProxyParsed.AddUObject(this, &USubtitlesAudioSubsystem::OnDialogueSoundWaveProxyParsed);
+}
+
+void USubtitlesAudioSubsystem::Deinitialize()
+{
+	USoundNodeDialoguePlayer::OnDialogueSoundWaveProxyParsed.RemoveAll(this);
+	Super::Deinitialize();
+}
+
 void USubtitlesAudioSubsystem::NotifyActiveSoundCreated(FActiveSound& ActiveSound)
 {
 	USoundBase* Sound = ActiveSound.GetSound();
@@ -27,50 +38,17 @@ void USubtitlesAudioSubsystem::NotifyActiveSoundCreated(FActiveSound& ActiveSoun
 			}
 		}
 
-		// Also check for DialogueWaves with localized text
-		// DialogueWaves, even inside SoundCues, should only play subtitles this way if the migration cvar was set.
-		// Since this callback is run when the sound is created in the sound engine, changing the cvar doesn't effect currently-playing sounds.
+		// Handle DialogueWaves played directly here.
+		// DialogueWaves inside SoundCues are handled via the OnDialogueSoundWaveProxyParsed delegate (see SoundNodeDialoguePlayer::ParseNodes).
 		static const IConsoleVariable* CVarUseNewSubtitles = IConsoleManager::Get().FindConsoleVariable(TEXT("au.UseNewSubtitles"));
 
 		if (CVarUseNewSubtitles && CVarUseNewSubtitles->GetBool() && ActiveSound.bHandleSubtitles)
 		{
-			FSubtitleAssetData SubtitleCue;
-
 			const UDialogueSoundWaveProxy* DialogueWave = Cast<UDialogueSoundWaveProxy>(Sound);
 			if (IsValid(DialogueWave))
 			{
 				QueueSubtitleFromDialogueWaveProxy(DialogueWave, ActiveSound);
 			}
-			else
-			{
-				// Check inside Sound Cues for DialogueWaves as well.
-				const USoundCue* SoundCue = Cast<USoundCue>(Sound);
-
-				// SoundCue's FirstNode is used to grab all the other nodes, so it needs to be valid also.
-				if (IsValid(SoundCue) && IsValid(SoundCue->FirstNode))
-				{
-					const UDialogueSoundWaveProxy* Proxy = CheckSoundCueNodeForDialogueWaves(SoundCue->FirstNode, ActiveSound);
-
-					if (IsValid(Proxy))
-					{
-						QueueSubtitleFromDialogueWaveProxy(Proxy, ActiveSound);
-					}
-					else
-					{
-						// If the parent node is not a DialogueWavePlayer, check for nested nodes that may also have DialogueWaves
-						TArray<USoundNode*> ChildNodes;
-						SoundCue->FirstNode->GetAllNodes(ChildNodes);
-						for (const USoundNode* Node : ChildNodes)
-						{
-							const UDialogueSoundWaveProxy* ChildProxy = CheckSoundCueNodeForDialogueWaves(Node, ActiveSound);
-							if (IsValid(ChildProxy))
-							{
-								QueueSubtitleFromDialogueWaveProxy(ChildProxy, ActiveSound);
-							}
-						}
-					}
-				}
-			}
 		}
 	}
 }
@@ -98,74 +76,71 @@ void USubtitlesAudioSubsystem::NotifyActiveSoundDeleting(const FActiveSound& Act
 			}
 		}
 
-		// Also check for DialogueWaves with localized text
 		static const IConsoleVariable* CVarUseNewSubtitles = IConsoleManager::Get().FindConsoleVariable(TEXT("au.UseNewSubtitles"));
 
 		if (CVarUseNewSubtitles && CVarUseNewSubtitles->GetBool() && ActiveSound.bHandleSubtitles)
 		{
-			FSubtitleAssetData SubtitleCue;
-
 			const UDialogueSoundWaveProxy* DialogueWave = Cast<UDialogueSoundWaveProxy>(Sound);
 			if (IsValid(DialogueWave))
 			{
 				StopSubtitleFromDialogueWaveProxy(DialogueWave, ActiveSound);
 			}
-			else
+			else if (TArray<FSubtitleAssetData>* Subtitles = QueuedSubtitlesByPlayOrder.Find(ActiveSound.GetPlayOrder()))
 			{
-				// Check inside Sound Cues for DialogueWaves as well.
-				const USoundCue* SoundCue = Cast<USoundCue>(Sound);
-
-				// SoundCue's FirstNode is used to grab all the other nodes, so it needs to be valid also.
-				if (IsValid(SoundCue) && IsValid(SoundCue->FirstNode))
+				for (const FSubtitleAssetData& Subtitle : *Subtitles)
 				{
-					const UDialogueSoundWaveProxy* Proxy = CheckSoundCueNodeForDialogueWaves(SoundCue->FirstNode, ActiveSound);
-
-					if (IsValid(Proxy))
-					{
-						// A dialogueWave won't reference another dialoguewave, so no need to search recursively
-						StopSubtitleFromDialogueWaveProxy(Proxy, ActiveSound);
-					}
-					else
-					{
-						// If the parent node is not a DialogueWavePlayer, check for nested nodes that may also have DialogueWaves
-						TArray<USoundNode*> ChildNodes;
-						SoundCue->FirstNode->GetAllNodes(ChildNodes);
-						for (const USoundNode* Node : ChildNodes)
-						{
-							const UDialogueSoundWaveProxy* ChildProxy = CheckSoundCueNodeForDialogueWaves(Node, ActiveSound);
-							if (IsValid(ChildProxy))
-							{
-								StopSubtitleFromDialogueWaveProxy(ChildProxy, ActiveSound);
-							}
-						}
-					}
+					FSubtitlesAndClosedCaptionsDelegates::StopSubtitle.ExecuteIfBound(Subtitle);
 				}
+				QueuedSubtitlesByPlayOrder.Remove(ActiveSound.GetPlayOrder());
 			}
 		}
 	}
 }
 
-const UDialogueSoundWaveProxy* USubtitlesAudioSubsystem::CheckSoundCueNodeForDialogueWaves(const USoundNode* Node, const FActiveSound& ActiveSound)
+void USubtitlesAudioSubsystem::OnDialogueSoundWaveProxyParsed(const FActiveSound& ActiveSound, const UDialogueSoundWaveProxy* Proxy)
 {
-	const USoundNodeDialoguePlayer* DialoguePlayerNode = Cast<USoundNodeDialoguePlayer>(Node);
+	if (ActiveSound.bHandleSubtitles)
+	{
+		ExecuteOnGameThread(
+			TEXT("USubtitlesAudioSubsystem::QueueSubtitleFromParsedSoundWaveProxy"),
 
-	// These pointers might be null or invalid:
-	// A SoundCue might have a Dialogue node that doesn't actually reference any assets, or the asset itself might lack any actual soundwaves to play.
-	if (IsValid(DialoguePlayerNode))
+			[WeakThis = TWeakObjectPtr<USubtitlesAudioSubsystem>(this),
+			Text = Proxy->GetSubtitleText(), Priority = ActiveSound.SubtitlePriority, Duration = Proxy->GetDuration(),
+			StartOffset = ActiveSound.RequestedStartTime, PlayOrder = ActiveSound.GetPlayOrder()]()
+			{
+				if (TStrongObjectPtr<USubtitlesAudioSubsystem> ThisPin = WeakThis.Pin())
+				{
+					ThisPin->QueueSubtitleFromParsedSoundWaveProxy(Text, Priority, Duration, StartOffset, PlayOrder);
+				}
+			}
+		);
+	}
+}
+
+void USubtitlesAudioSubsystem::QueueSubtitleFromParsedSoundWaveProxy(const FText Text, const float Priority, const float Duration, const float StartOffset, const uint32 PlayOrder)
+{
+	// Called from the audio thread, executed on game thread: fired synchronously from Broadcast inside SoundNodeDialoguePlayer::ParseNodes.
+	static const IConsoleVariable* CVarUseNewSubtitles = IConsoleManager::Get().FindConsoleVariable(TEXT("au.UseNewSubtitles"));
+
+	if (!CVarUseNewSubtitles || !CVarUseNewSubtitles->GetBool())
 	{
-		const UDialogueSoundWaveProxy* Proxy = DialoguePlayerNode->GetContextualDialogueSoundWaveProxy();
-		if (IsValid(Proxy))
-		{
-			return Proxy;
-		}
+		return;
 	}
 
-	return nullptr;
+	FSubtitleAssetData SubtitleCue;
+	SubtitleCue.Text = Text;
+	SubtitleCue.Priority = Priority;
+	SubtitleCue.Duration = Duration;
+	SubtitleCue.StartOffset = StartOffset;
+
+	QueuedSubtitlesByPlayOrder.FindOrAdd(PlayOrder).Add(SubtitleCue);
+
+	FSubtitlesAndClosedCaptionsDelegates::QueueSubtitle.ExecuteIfBound(FQueueSubtitleParameters{ SubtitleCue }, ESubtitleTiming::InternallyTimed);
 }
 
 void USubtitlesAudioSubsystem::QueueSubtitleFromDialogueWaveProxy(const UDialogueSoundWaveProxy* DialogueWaveProxy, const FActiveSound& ActiveSound)
 {
-	check(ActiveSound.bHandleSubtitles);	// Function should only be called if the Active Sound actually has subtitles it needs queued.
+	check(ActiveSound.bHandleSubtitles); // Function should only be called if the Active Sound actually has subtitles it needs queued.
 	FSubtitleAssetData SubtitleCue;
 	SubtitleCue.Text = DialogueWaveProxy->GetSubtitleText();
 	SubtitleCue.Priority = ActiveSound.SubtitlePriority;
@@ -177,7 +152,7 @@ void USubtitlesAudioSubsystem::QueueSubtitleFromDialogueWaveProxy(const UDialogu
 
 void USubtitlesAudioSubsystem::StopSubtitleFromDialogueWaveProxy(const UDialogueSoundWaveProxy* DialogueWaveProxy, const FActiveSound& ActiveSound)
 {
-	check(ActiveSound.bHandleSubtitles);	// Function should only be called if the Active Sound actually has subtitles it needs queued.
+	check(ActiveSound.bHandleSubtitles); // Function should only be called if the Active Sound actually has subtitles it needs queued.
 	FSubtitleAssetData SubtitleCue;
 	SubtitleCue.Text = DialogueWaveProxy->GetSubtitleText();
 	SubtitleCue.Priority = ActiveSound.SubtitlePriority;
diff --git a/Engine/Plugins/SubtitlesAndClosedCaptions/Source/SubtitlesAndClosedCaptions/Public/SubtitlesAudioSubsystem.h b/Engine/Plugins/SubtitlesAndClosedCaptions/Source/SubtitlesAndClosedCaptions/Public/SubtitlesAudioSubsystem.h
--- a/Engine/Plugins/SubtitlesAndClosedCaptions/Source/SubtitlesAndClosedCaptions/Public/SubtitlesAudioSubsystem.h
+++ b/Engine/Plugins/SubtitlesAndClosedCaptions/Source/SubtitlesAndClosedCaptions/Public/SubtitlesAudioSubsystem.h
@@ -1,14 +1,14 @@
 // Copyright Epic Games, Inc. All Rights Reserved.
 
-#pragma once 
+#pragma once
 #include "Subsystems/AudioEngineSubsystem.h"
 #include "ActiveSoundUpdateInterface.h"
+#include "Subtitles/SubtitlesAndClosedCaptionsTypes.h"
 #include "SubtitlesAudioSubsystem.generated.h"
 
 #define UE_API SUBTITLESANDCLOSEDCAPTIONS_API
 
 class UDialogueSoundWaveProxy;
-class USoundNode;
 
 /*
 *	SubtitlesAudioSubsystem - AudioEngineSubsystem for automatically queueing subtitles on audio classes.
@@ -24,17 +24,29 @@ class USubtitlesAudioSubsystem : public UAudioEngineSubsystem
 public:
 	USubtitlesAudioSubsystem() = default;
 
+	//~ Begin USubsystem Interface
+	virtual void Initialize(FSubsystemCollectionBase& Collection) override;
+	virtual void Deinitialize() override;
+	//~ End USubsystem Interface
+
 	//~ Begin IActiveSoundUpdateInterface
 	virtual void NotifyActiveSoundCreated(FActiveSound& ActiveSound) override;
 	virtual void NotifyActiveSoundDeleting(const FActiveSound& ActiveSound) override;
 	//~ End IActiveSoundUpdateInterface
 
 private:
-	// Helper functions to cut down on duplicated code
-	static const UDialogueSoundWaveProxy* CheckSoundCueNodeForDialogueWaves(const USoundNode* Node, const FActiveSound& ActiveSound);
-	static void QueueSubtitleFromDialogueWaveProxy(const UDialogueSoundWaveProxy* DialogueWaveProxy, const FActiveSound& ActiveSound);
+	// Called on the audio thread when a SoundNodeDialoguePlayer first resolves its proxy for a given active sound.
+	void OnDialogueSoundWaveProxyParsed(const FActiveSound& ActiveSound, const UDialogueSoundWaveProxy* Proxy);
+
+	// Game-thread side of the same function
+	void QueueSubtitleFromParsedSoundWaveProxy(const FText Text, const float Priority, const float Duration, const float StartOffset, const uint32 PlayOrder);
 
+	static void QueueSubtitleFromDialogueWaveProxy(const UDialogueSoundWaveProxy* DialogueWaveProxy, const FActiveSound& ActiveSound);
 	static void StopSubtitleFromDialogueWaveProxy(const UDialogueSoundWaveProxy* DialogueWaveProxy, const FActiveSound& ActiveSound);
+
+	// Tracks which subtitle data is queued for each active sound. PlayOrder, the key, is monotonically incremented by FActiveSound and is thus unique.
+	// Used later by NotifyActiveSoundDeleting for stopping the associated subtitles.
+	TMap<uint32, TArray<FSubtitleAssetData>> QueuedSubtitlesByPlayOrder;
 };
 
 #undef UE_API
diff --git a/Engine/Source/Runtime/Engine/Classes/Sound/SoundNodeDialoguePlayer.h b/Engine/Source/Runtime/Engine/Classes/Sound/SoundNodeDialoguePlayer.h
--- a/Engine/Source/Runtime/Engine/Classes/Sound/SoundNodeDialoguePlayer.h
+++ b/Engine/Source/Runtime/Engine/Classes/Sound/SoundNodeDialoguePlayer.h
@@ -14,6 +14,8 @@ struct FActiveSound;
 struct FSoundParseParameters;
 struct FWaveInstance;
 
+DECLARE_TS_MULTICAST_DELEGATE_TwoParams(FOnDialogueSoundWaveProxyParsed, const FActiveSound&, const UDialogueSoundWaveProxy*);
+
 /**
  Sound node that contains a reference to the dialogue table to pull from and be played
 */
@@ -37,6 +39,10 @@ public:
 	// Subtitles need to be able to grab the UDialogueSoundWaveProxy to for localized text and duration.
 	ENGINE_API const UDialogueSoundWaveProxy* GetContextualDialogueSoundWaveProxy() const;
 
+	// Fired on the audio thread the first time ParseNodes resolves a valid proxy for a given active sound.
+	// Listeners can use this to queue subtitles only for dialogue waves that are actually playing.
+	static ENGINE_API FOnDialogueSoundWaveProxyParsed OnDialogueSoundWaveProxyParsed;
+
 	//~ Begin USoundNode Interface
 	virtual int32 GetMaxChildNodes() const override;
 	virtual float GetDuration();
diff --git a/Engine/Source/Runtime/Engine/Private/SoundNodeDialoguePlayer.cpp b/Engine/Source/Runtime/Engine/Private/SoundNodeDialoguePlayer.cpp
--- a/Engine/Source/Runtime/Engine/Private/SoundNodeDialoguePlayer.cpp
+++ b/Engine/Source/Runtime/Engine/Private/SoundNodeDialoguePlayer.cpp
@@ -5,6 +5,8 @@
 #include "Sound/DialogueWave.h"
 #include "Sound/DialogueSoundWaveProxy.h"
 
+FOnDialogueSoundWaveProxyParsed USoundNodeDialoguePlayer::OnDialogueSoundWaveProxyParsed;
+
 #include UE_INLINE_GENERATED_CPP_BY_NAME(SoundNodeDialoguePlayer)
 
 #define LOCTEXT_NAMESPACE "SoundNodeDialoguePlayer"
@@ -19,6 +21,16 @@ void USoundNodeDialoguePlayer::ParseNodes( FAudioDevice* AudioDevice, const UPTR
 	USoundBase* SoundBase = GetDialogueWave() ? GetDialogueWave()->GetWaveFromContext(DialogueWaveParameter.Context) : NULL;
 	if (SoundBase)
 	{
+		RETRIEVE_SOUNDNODE_PAYLOAD(sizeof(uint8));
+		if (*RequiresInitialization)
+		{
+			*RequiresInitialization = 0;
+			if (const UDialogueSoundWaveProxy* Proxy = Cast<UDialogueSoundWaveProxy>(SoundBase))
+			{
+				OnDialogueSoundWaveProxyParsed.Broadcast(ActiveSound, Proxy);
+			}
+		}
+
 		if (bLooping)
 		{
 			FSoundParseParameters UpdatedParams = ParseParams;
