Change SubtitlesSubsystem Error to Warning to unblock Publish Moderation
#Synced-CL 52180472

diff --git a/Engine/Plugins/SubtitlesAndClosedCaptions/Source/SubtitlesAndClosedCaptions/Private/SubtitlesSubsystem.cpp b/Engine/Plugins/SubtitlesAndClosedCaptions/Source/SubtitlesAndClosedCaptions/Private/SubtitlesSubsystem.cpp
--- a/Engine/Plugins/SubtitlesAndClosedCaptions/Source/SubtitlesAndClosedCaptions/Private/SubtitlesSubsystem.cpp
+++ b/Engine/Plugins/SubtitlesAndClosedCaptions/Source/SubtitlesAndClosedCaptions/Private/SubtitlesSubsystem.cpp
@@ -339,7 +339,7 @@ void USubtitlesSubsystem::ReplaceWidget(const TSoftClassPtr<USubtitleWidget>& Ne
 
 	if (!TryCreateUMGWidgetFromAsset(NewWidgetAsset))
 	{
-		UE_LOGFMT(LogSubtitlesAndClosedCaptions, Error, "Can't replace the subtitle widget; was a valid asset provided?");
+		UE_LOGFMT(LogSubtitlesAndClosedCaptions, Warning, "Can't replace the subtitle widget; was a valid asset provided?");
 	}
 }
 
